New Tool: Hierarchy Editor

diff --git a/Engine/Source/Editor/HierarchyEditor/Private/HierarchyViewModelBase.cpp b/Engine/Source/Editor/HierarchyEditor/Private/HierarchyViewModelBase.cpp
new file mode 100644
--- /dev/null
+++ b/Engine/Source/Editor/HierarchyEditor/Private/HierarchyViewModelBase.cpp
@@ -0,0 +1,1823 @@
+// Copyright Epic Games, Inc. All Rights Reserved.
+
+#include "HierarchyViewModelBase.h"
+#include "SDropTarget.h"
+#include "Widgets/Views/STableRow.h"
+#include "Widgets/SHierarchyEditor.h"
+#include "Widgets/Text/SInlineEditableTextBlock.h"
+#include "Widgets/SToolTip.h"
+#include "Widgets/Input/SCheckBox.h"
+#include "Editor.h"
+#include "HierarchyEditorCommands.h"
+#include "HierarchyEditorMisc.h"
+#include "HierarchyEditorModule.h"
+#include "IPropertyRowGenerator.h"
+#include "ScopedTransaction.h"
+#include "Framework/Commands/GenericCommands.h"
+#include "ScopedTransaction.h"
+#include "Framework/Notifications/NotificationManager.h"
+#include "ToolMenus.h"
+#include "Widgets/Notifications/SNotificationList.h"
+
+#define LOCTEXT_NAMESPACE "HierarchyEditor"
+
+const TArray<UHierarchyElement*>& UHierarchyElement::GetChildren() const
+{
+	return Children;
+}
+
+UHierarchyElement* UHierarchyElement::FindChildWithIdentity(FHierarchyElementIdentity ChildIdentity, bool bSearchRecursively)
+{
+	TObjectPtr<UHierarchyElement>* FoundItem = Children.FindByPredicate([ChildIdentity](UHierarchyElement* Child)
+	{
+		return Child->GetPersistentIdentity() == ChildIdentity;
+	});
+
+	if(FoundItem)
+	{
+		return *FoundItem;
+	}
+	
+	if(bSearchRecursively)
+	{
+		for(UHierarchyElement* Child : Children)
+		{
+			UHierarchyElement* FoundChild = Child->FindChildWithIdentity(ChildIdentity, bSearchRecursively);
+
+			if(FoundChild)
+			{
+				return FoundChild;
+			}
+		}
+	}	
+
+	return nullptr;
+}
+
+UHierarchyElement* UHierarchyElement::CopyAndAddItemAsChild(const UHierarchyElement& ItemToCopy)
+{
+	UHierarchyElement* NewChild = Cast<UHierarchyElement>(StaticDuplicateObject(&ItemToCopy, this));
+	if(NewChild->GetPersistentIdentity() != ItemToCopy.GetPersistentIdentity())
+	{
+		check(false);
+	}
+	GetChildrenMutable().Add(NewChild);
+
+	return NewChild;
+}
+
+UHierarchyElement* UHierarchyElement::CopyAndAddItemUnderParentIdentity(const UHierarchyElement& ItemToCopy, FHierarchyElementIdentity ParentIdentity)
+{
+	UHierarchyElement* ParentItem = FindChildWithIdentity(ParentIdentity, true);
+
+	if(ParentItem)
+	{
+		UHierarchyElement* NewChild = Cast<UHierarchyElement>(StaticDuplicateObject(&ItemToCopy, ParentItem));
+		if(NewChild->GetPersistentIdentity() != ItemToCopy.GetPersistentIdentity())
+		{
+			check(false);
+		}
+		ParentItem->GetChildrenMutable().Add(NewChild);
+		return NewChild;
+	}
+
+	return nullptr;
+}
+
+bool UHierarchyElement::RemoveChildWithIdentity(FHierarchyElementIdentity ChildIdentity, bool bSearchRecursively)
+{	
+	int32 RemovedChildrenCount = Children.RemoveAll([ChildIdentity](UHierarchyElement* Child)
+	{
+		return Child->GetPersistentIdentity() == ChildIdentity;
+	});
+
+	if(RemovedChildrenCount > 1)
+	{
+		UE_LOG(LogHierarchyEditor, Warning, TEXT("More than one child with the same identity has been found in parent %s"), *ToString());
+	}
+	
+	bool bChildrenRemoved = RemovedChildrenCount > 0;
+	
+	if(bSearchRecursively && bChildrenRemoved == false)
+	{
+		for(UHierarchyElement* Child : Children)
+		{
+			bChildrenRemoved |= Child->RemoveChildWithIdentity(ChildIdentity, bSearchRecursively);
+		}
+	}
+
+	return bChildrenRemoved;
+}
+
+TArray<FHierarchyElementIdentity> UHierarchyElement::GetParentIdentities() const
+{
+	TArray<FHierarchyElementIdentity> ParentIdentities;
+
+	for(UHierarchyElement* Parent = Cast<UHierarchyElement>(GetOuter()); Parent != nullptr; Parent = Cast<UHierarchyElement>(Parent->GetOuter()))
+	{
+		ParentIdentities.Add(Parent->GetPersistentIdentity());
+	}
+
+	return ParentIdentities;
+}
+
+bool UHierarchyElement::Modify(bool bAlwaysMarkDirty)
+{
+	bool bSavedToTransactionBuffer = true;
+
+	for(UHierarchyElement* Child : Children)
+	{
+		bSavedToTransactionBuffer &= Child->Modify(bAlwaysMarkDirty);
+	}
+	
+	bSavedToTransactionBuffer &= UObject::Modify(bAlwaysMarkDirty);
+
+	return bSavedToTransactionBuffer;
+}
+
+void UHierarchyElement::PostLoad()
+{
+	if(Guid_DEPRECATED.IsValid())
+	{
+		SetIdentity(FHierarchyElementIdentity({Guid_DEPRECATED}, {}));
+	}
+
+	Super::PostLoad();
+}
+
+UHierarchySection* UHierarchyRoot::AddSection(FText InNewSectionName, int32 InsertIndex, TSubclassOf<UHierarchySection> SectionClass)
+{
+	TSet<FName> ExistingSectionNames;
+	
+	for(FName& SectionName : GetSections())
+	{
+		ExistingSectionNames.Add(SectionName);
+	}
+	
+	FName NewName = UE::HierarchyEditor::GetUniqueName(FName(InNewSectionName.ToString()), ExistingSectionNames);
+	UHierarchySection* NewSectionItem = NewObject<UHierarchySection>(this, SectionClass);
+	NewSectionItem->SetSectionName(NewName);
+	NewSectionItem->SetFlags(RF_Transactional);
+	
+	if(InsertIndex == INDEX_NONE)
+	{
+		Sections.Add(NewSectionItem);
+	}
+	else
+	{
+		Sections.Insert(NewSectionItem, InsertIndex);
+	}
+	
+	return NewSectionItem;
+}
+
+UHierarchySection* UHierarchyRoot::FindSectionByIdentity(FHierarchyElementIdentity SectionIdentity)
+{
+	for(UHierarchySection* Section : Sections)
+	{
+		if(Section->GetPersistentIdentity() == SectionIdentity)
+		{
+			return Section;
+		}
+	}
+
+	return nullptr;
+}
+
+void UHierarchyRoot::DuplicateSectionFromOtherRoot(const UHierarchySection& SectionToCopy)
+{
+	if(FindSectionByIdentity(SectionToCopy.GetPersistentIdentity()) != nullptr || SectionToCopy.GetOuter() == this)
+	{
+		return;
+	}
+	
+	Sections.Add(Cast<UHierarchySection>(StaticDuplicateObject(&SectionToCopy, this)));
+}
+
+void UHierarchyRoot::RemoveSection(FText SectionName)
+{
+	if(Sections.ContainsByPredicate([SectionName](UHierarchySection* Section)
+	{
+		return Section->GetSectionNameAsText().EqualTo(SectionName);
+	}))
+	{
+		Sections.RemoveAll([SectionName](UHierarchySection* Section)
+		{
+			return Section->GetSectionNameAsText().EqualTo(SectionName);
+		});
+	}
+}
+
+void UHierarchyRoot::RemoveSectionByIdentity(FHierarchyElementIdentity SectionIdentity)
+{
+	Sections.RemoveAll([SectionIdentity](UHierarchySection* Section)
+	{
+		return Section->GetPersistentIdentity() == SectionIdentity;
+	});
+}
+
+TSet<FName> UHierarchyRoot::GetSections() const
+{
+	TSet<FName> OutSections;
+	for(UHierarchySection* Section : Sections)
+	{
+		OutSections.Add(Section->GetSectionName());
+	}
+
+	return OutSections;
+}
+
+int32 UHierarchyRoot::GetSectionIndex(UHierarchySection* Section) const
+{
+	return Sections.Find(Section);
+}
+
+bool UHierarchyRoot::Modify(bool bAlwaysMarkDirty)
+{
+	bool bSavedToTransactionBuffer = true;
+	
+	for(UHierarchySection* Section : Sections)
+	{
+		bSavedToTransactionBuffer &= Section->Modify();
+	}
+	
+	bSavedToTransactionBuffer &= Super::Modify(bAlwaysMarkDirty);	
+
+	return bSavedToTransactionBuffer;
+}
+
+void UHierarchyRoot::EmptyAllData()
+{
+	Children.Empty();
+	Sections.Empty();
+}
+
+bool FHierarchyCategoryViewModel::IsTopCategoryActive() const
+{
+	if(UHierarchyCategory* Category = GetDataMutable<UHierarchyCategory>())
+	{
+		const UHierarchyCategory* Result = Category;
+		const UHierarchyCategory* TopLevelCategory = Result;
+		
+		for (; TopLevelCategory != nullptr; TopLevelCategory = TopLevelCategory->GetTypedOuter<UHierarchyCategory>() )
+		{
+			if(TopLevelCategory != nullptr)
+			{
+				Result = TopLevelCategory;
+			}
+		}
+		
+		return HierarchyViewModel->IsHierarchySectionActive(Result->GetSection());
+	}
+
+	return false;	
+}
+
+FHierarchyElementViewModel::FCanPerformActionResults FHierarchyCategoryViewModel::CanDropOnInternal(TSharedPtr<FHierarchyElementViewModel> DraggedElement, EItemDropZone ItemDropZone)
+{
+	FCanPerformActionResults Results(false);
+	
+	TArray<TSharedPtr<FHierarchyCategoryViewModel>> TargetChildrenCategories;
+	GetChildrenViewModelsForType<UHierarchyCategory, FHierarchyCategoryViewModel>(TargetChildrenCategories);
+	
+	TArray<TSharedPtr<FHierarchyCategoryViewModel>> SiblingCategories;
+	Parent.Pin()->GetChildrenViewModelsForType<UHierarchyCategory, FHierarchyCategoryViewModel>(SiblingCategories);
+	
+	// we only allow drops if some general conditions are fulfilled
+	if(DraggedElement->GetData() != GetData() &&
+		(!DraggedElement->HasParent(AsShared(), false) || ItemDropZone != EItemDropZone::OntoItem)  &&
+		!HasParent(DraggedElement, true))
+	{
+		// categories can be dropped on categories, but only if the resulting sibling categories or children categories have different names
+		if(DraggedElement->GetData()->IsA<UHierarchyCategory>())
+		{
+			if(ItemDropZone != EItemDropZone::OntoItem)
+			{
+				bool bContainsSiblingWithSameName = SiblingCategories.ContainsByPredicate([DraggedElement](TSharedPtr<FHierarchyCategoryViewModel> HierarchyCategoryViewModel)
+					{
+						return DraggedElement->ToString() == HierarchyCategoryViewModel->ToString() && DraggedElement != HierarchyCategoryViewModel;
+					});
+
+				if(bContainsSiblingWithSameName)
+				{
+					Results.bCanPerform = false;
+					Results.CanPerformMessage = LOCTEXT("CantDropCategorNextToCategorySameSiblingNames", "A category of the same name already exists here, potentially in a different section. Please rename your category first.");
+					return Results;
+				}
+
+				Results.CanPerformMessage = LOCTEXT("MoveCategoryText", "Move category here");
+
+				// if we are making a category a sibling of another at the root level, the section will be set to the currently active section. Let that be known.
+				if(Parent.Pin()->GetData()->IsA<UHierarchyRoot>())
+				{
+					UHierarchyCategory* DraggedCategory = Cast<UHierarchyCategory>(DraggedElement->GetDataMutable());
+					if(DraggedCategory->GetSection() != HierarchyViewModel->GetActiveHierarchySectionData())
+					{
+						FText SectionChangeBaseText = LOCTEXT("CategorySectionWillUpdateDueToDrop", "The section of the category will change to {0} after the drop");
+						FText ActualSectionChangeText = FText::FormatOrdered(SectionChangeBaseText, HierarchyViewModel->GetActiveHierarchySectionData() == nullptr ? FText::FromString("All") : HierarchyViewModel->GetActiveHierarchySectionData()->GetSectionNameAsText());
+						Results.CanPerformMessage = FText::FormatOrdered(FText::AsCultureInvariant("{0}\n{1}"), Results.CanPerformMessage, ActualSectionChangeText);
+					}
+				}
+			}
+			else
+			{
+				bool bContainsChildrenCategoriesWithSameName = TargetChildrenCategories.ContainsByPredicate([DraggedElement](TSharedPtr<FHierarchyCategoryViewModel> HierarchyCategoryViewModel)
+					{
+						return DraggedElement->ToString() == HierarchyCategoryViewModel->ToString();
+					});
+
+				if(bContainsChildrenCategoriesWithSameName)
+				{
+					Results.bCanPerform = false;
+					Results.CanPerformMessage = LOCTEXT("CantDropCategoryOnCategorySameChildCategoryName", "A sub-category of the same name already exists! Please rename your category first.");
+					return Results;
+				}
+
+				Results.CanPerformMessage = LOCTEXT("CreateSubcategory", "Drop category here to create a sub-category");
+			}
+
+			Results.bCanPerform = true;
+			return Results;
+		}
+		else if(DraggedElement->GetData()->IsA<UHierarchyItem>())
+		{
+			// items can generally be dropped onto categories
+			Results.bCanPerform = EItemDropZone::OntoItem == ItemDropZone;
+
+			if(Results.bCanPerform)
+			{
+				if(DraggedElement->IsForHierarchy() == false)
+				{
+					FText Message = LOCTEXT("AddItemToCategoryDragMessage", "Add {0} to {1}");
+					Results.CanPerformMessage = FText::FormatOrdered(Message, FText::FromString(DraggedElement->ToString()), FText::FromString(ToString()));
+				}
+				else
+				{
+					FText Message = LOCTEXT("MoveItemToCategoryDragMessage", "Move {0} to {1}");
+					Results.CanPerformMessage = FText::FormatOrdered(Message, FText::FromString(DraggedElement->ToString()), FText::FromString(ToString()));
+				}
+			}
+		}
+	}
+
+	return Results;
+}
+
+void UHierarchyCategory::FixupSectionLinkage()
+{
+	UHierarchyRoot* OwningRoot = GetTypedOuter<UHierarchyRoot>();
+
+	if(Section != nullptr && Section->GetTypedOuter<UHierarchyRoot>() != OwningRoot)
+	{
+		UHierarchySection* CorrectSection = OwningRoot->FindSectionByIdentity(Section->GetPersistentIdentity());
+		ensure(CorrectSection != nullptr);
+		Section = CorrectSection;
+	}
+}
+
+void UHierarchySection::SetSectionNameAsText(const FText& Text)
+{
+	Section = FName(Text.ToString());
+}
+
+UHierarchyViewModelBase::UHierarchyViewModelBase()
+{
+	Commands = MakeShared<FUICommandList>();
+}
+
+UHierarchyViewModelBase::~UHierarchyViewModelBase()
+{
+	RefreshSourceViewDelegate.Unbind();
+	RefreshHierarchyWidgetDelegate.Unbind();
+	RefreshSectionsViewDelegate.Unbind();
+}
+
+void UHierarchyViewModelBase::Initialize()
+{		
+	HierarchyRoot = GetHierarchyRoot();
+	HierarchyRoot->SetFlags(RF_Transactional);
+
+	TArray<UHierarchyElement*> AllItems;
+	HierarchyRoot->GetChildrenOfType<UHierarchyElement>(AllItems, true);
+	for(UHierarchyElement* Item : AllItems)
+	{
+		Item->SetFlags(RF_Transactional);
+	}
+
+	for(UHierarchySection* Section : HierarchyRoot->GetSectionDataMutable())
+	{
+		Section->SetFlags(RF_Transactional);
+	}
+
+	UToolMenus* ToolMenus = UToolMenus::Get();
+
+	FName MenuName("HierarchyMenu");
+	if(!ToolMenus->IsMenuRegistered(MenuName))
+	{
+		UToolMenu* HierarchyMenu = ToolMenus->RegisterMenu(MenuName, NAME_None, EMultiBoxType::Menu);
+		HierarchyMenu->AddDynamicSection(NAME_None, FNewToolMenuDelegate::CreateUObject(this, &UHierarchyViewModelBase::GenerateDynamicContextMenu));
+	}
+	
+	SetupCommands();
+
+	TSharedPtr<FHierarchyElementViewModel> ViewModel = CreateViewModelForElement(HierarchyRoot, nullptr);
+	HierarchyRootViewModel = StaticCastSharedPtr<FHierarchyRootViewModel>(ViewModel); 
+	if(!ensureMsgf(HierarchyRootViewModel.IsValid(), TEXT("Make sure that CreateViewModelForData creates a FHierarchyRootViewModel (or derived) for UHierarchyRoot elements")))
+	{
+		return;
+	}
+	
+	HierarchyRootViewModel->Initialize();
+	HierarchyRootViewModel->AddChildFilter(FHierarchyElementViewModel::FOnFilterChild::CreateUObject(this, &UHierarchyViewModelBase::FilterForHierarchySection));
+	HierarchyRootViewModel->AddChildFilter(FHierarchyElementViewModel::FOnFilterChild::CreateUObject(this, &UHierarchyViewModelBase::FilterForUncategorizedRootItemsInAllSection));
+	HierarchyRootViewModel->SyncViewModelsToData();
+	
+	SetActiveHierarchySection(nullptr);
+	
+	InitializeInternal();
+
+	bIsInitialized = true;
+	
+	OnInitializedDelegate.ExecuteIfBound();
+}
+
+void UHierarchyViewModelBase::Finalize()
+{	
+	HierarchyRootViewModel.Reset();
+	HierarchyRoot = nullptr;
+	
+	FinalizeInternal();
+	
+	bIsFinalized = true;
+}
+
+TSharedPtr<FHierarchyElementViewModel> UHierarchyViewModelBase::CreateViewModelForElement(UHierarchyElement* Element, TSharedPtr<FHierarchyElementViewModel> Parent)
+{
+	// We first give the internal implementation a chance to create view models
+	if(TSharedPtr<FHierarchyElementViewModel> CustomViewModel = CreateCustomViewModelForElement(Element, Parent))
+	{
+		return CustomViewModel;
+	}
+
+	// If it wasn't implemented or wasn't covered, we make sure to have default view models
+	if(UHierarchyItem* Item = Cast<UHierarchyItem>(Element))
+	{
+		return MakeShared<FHierarchyItemViewModel>(Item, Parent.ToSharedRef(), this);
+	}
+	else if(UHierarchyCategory* Category = Cast<UHierarchyCategory>(Element))
+	{
+		return MakeShared<FHierarchyCategoryViewModel>(Category, Parent.ToSharedRef(), this);
+	}
+	else if(UHierarchySection* Section = Cast<UHierarchySection>(Element))
+	{
+		// For sections, we require the parent to be a root view model
+		TSharedPtr<FHierarchyRootViewModel> RootViewModel = StaticCastSharedPtr<FHierarchyRootViewModel>(Parent);
+		ensure(RootViewModel.IsValid());
+		return MakeShared<FHierarchySectionViewModel>(Section, RootViewModel.ToSharedRef(), this);
+	}
+	else if(UHierarchyRoot* Root = Cast<UHierarchyRoot>(Element))
+	{
+		// If the root is the hierarchy root, we know it's for the hierarchy. If not, it's the transient source root
+		bool bIsForHierarchy = GetHierarchyRoot() == Element;
+		return MakeShared<FHierarchyRootViewModel>(Root, this, bIsForHierarchy);
+	}
+
+	ensureMsgf(false, TEXT("This should never be reached. Either a custom or a default view model must exist for each Hierarchy Element"));
+	return nullptr;
+}
+
+TSubclassOf<UHierarchyCategory> UHierarchyViewModelBase::GetCategoryDataClass() const
+{
+	return UHierarchyCategory::StaticClass();
+}
+
+TSubclassOf<UHierarchySection> UHierarchyViewModelBase::GetSectionDataClass() const
+{
+	return UHierarchySection::StaticClass();
+}
+
+
+void UHierarchyViewModelBase::ForceFullRefresh()
+{
+	RefreshSourceItemsRequestedDelegate.ExecuteIfBound();
+	// todo (me) during merge at startup this can be nullptr for some reason
+	if(HierarchyRootViewModel.IsValid())
+	{
+		HierarchyRootViewModel->SyncViewModelsToData();
+	}
+	RefreshAllViewsRequestedDelegate.ExecuteIfBound(true);
+}
+
+void UHierarchyViewModelBase::ForceFullRefreshOnTimer()
+{
+	ensure(FullRefreshNextFrameHandle.IsValid());
+	ForceFullRefresh();
+	FullRefreshNextFrameHandle.Invalidate();
+}
+
+void UHierarchyViewModelBase::RequestFullRefreshNextFrame()
+{
+	if(!FullRefreshNextFrameHandle.IsValid() && GEditor != nullptr)
+	{
+		FullRefreshNextFrameHandle = GEditor->GetTimerManager()->SetTimerForNextTick(FTimerDelegate::CreateUObject(this, &UHierarchyViewModelBase::ForceFullRefreshOnTimer));
+	}
+}
+
+void UHierarchyViewModelBase::RefreshAllViews(bool bFullRefresh) const
+{
+	RefreshAllViewsRequestedDelegate.ExecuteIfBound(bFullRefresh);
+}
+
+void UHierarchyViewModelBase::RefreshSourceView(bool bFullRefresh) const
+{
+	RefreshSourceViewDelegate.ExecuteIfBound(bFullRefresh);
+}
+
+void UHierarchyViewModelBase::RefreshHierarchyView(bool bFullRefresh) const
+{
+	RefreshHierarchyWidgetDelegate.ExecuteIfBound(bFullRefresh);
+}
+
+void UHierarchyViewModelBase::RefreshSectionsView() const
+{
+	RefreshSectionsViewDelegate.ExecuteIfBound();
+}
+
+void UHierarchyViewModelBase::PostUndo(bool bSuccess)
+{
+	ForceFullRefresh();
+}
+
+void UHierarchyViewModelBase::PostRedo(bool bSuccess)
+{
+	PostUndo(bSuccess);
+}
+
+bool UHierarchyViewModelBase::MatchesContext(const FTransactionContext& InContext, const TArray<TPair<UObject*, FTransactionObjectEvent>>& TransactionObjectContexts) const
+{
+	for(const TPair<UObject*, FTransactionObjectEvent>& TransactionObjectContext : TransactionObjectContexts)
+	{
+		if(TransactionObjectContext.Key->IsA<UHierarchyElement>())
+		{
+			return true;
+		}
+	}
+	
+	return false;
+}
+
+bool UHierarchyViewModelBase::FilterForHierarchySection(TSharedPtr<const FHierarchyElementViewModel> ViewModel) const
+{
+	if(ActiveHierarchySection.IsValid())
+	{
+		return GetActiveHierarchySectionData() == ViewModel->GetSection();
+	}
+
+	return true;
+}
+
+bool UHierarchyViewModelBase::FilterForUncategorizedRootItemsInAllSection(TSharedPtr<const FHierarchyElementViewModel> ViewModel) const
+{
+	// we want to filter out all items that are directly added to the root if we aren't in the 'All' section
+	if(ActiveHierarchySection.IsValid())
+	{
+		return ViewModel->GetData<UHierarchyCategory>() != nullptr;
+	}
+
+	return true;
+}
+
+void UHierarchyViewModelBase::ToolMenuRequestRename(const FToolMenuContext& Context) const
+{
+	UHierarchyMenuContext* HierarchyMenuContext = Context.FindContext<UHierarchyMenuContext>();
+	if(HierarchyMenuContext->MenuHierarchyElements.Num() == 1)
+	{
+		HierarchyMenuContext->MenuHierarchyElements[0]->RequestRename();
+	}
+}
+
+bool UHierarchyViewModelBase::ToolMenuCanRequestRename(const FToolMenuContext& Context) const
+{
+	UHierarchyMenuContext* HierarchyMenuContext = Context.FindContext<UHierarchyMenuContext>();
+	if(HierarchyMenuContext->MenuHierarchyElements.Num() == 1)
+	{
+		return HierarchyMenuContext->MenuHierarchyElements[0]->CanRename();
+	}
+
+	return false;
+}
+
+void UHierarchyViewModelBase::ToolMenuDelete(const FToolMenuContext& Context) const
+{
+	UHierarchyMenuContext* HierarchyMenuContext = Context.FindContext<UHierarchyMenuContext>();
+	DeleteElements(HierarchyMenuContext->MenuHierarchyElements);
+}
+
+bool UHierarchyViewModelBase::ToolMenuCanDelete(const FToolMenuContext& Context) const
+{
+	UHierarchyMenuContext* HierarchyMenuContext = Context.FindContext<UHierarchyMenuContext>();
+	
+	for(TSharedPtr<FHierarchyElementViewModel> MenuHierarchyElement : HierarchyMenuContext->MenuHierarchyElements)
+	{
+		if(MenuHierarchyElement->CanDelete() == false)
+		{
+			return false;
+		}
+	}
+
+	return HierarchyMenuContext->MenuHierarchyElements.Num() > 0;
+}
+
+void UHierarchyViewModelBase::ToolMenuNavigateTo(const FToolMenuContext& Context) const
+{
+	UHierarchyMenuContext* HierarchyMenuContext = Context.FindContext<UHierarchyMenuContext>();
+
+	if(HierarchyMenuContext->MenuHierarchyElements.Num() == 1)
+	{
+		if(TSharedPtr<FHierarchyElementViewModel> MatchingViewModelInHierarchy = GetHierarchyRootViewModel()->FindViewModelForChild(HierarchyMenuContext->MenuHierarchyElements[0]->GetData()->GetPersistentIdentity(), true))
+		{
+			NavigateToElementInHierarchy(MatchingViewModelInHierarchy.ToSharedRef());
+		}
+	}
+}
+
+bool UHierarchyViewModelBase::ToolMenuCanNavigateTo(const FToolMenuContext& Context) const
+{
+	UHierarchyMenuContext* HierarchyMenuContext = Context.FindContext<UHierarchyMenuContext>();
+
+	if(HierarchyMenuContext->MenuHierarchyElements.Num() != 1)
+	{
+		return false;
+	}
+	
+	TSharedPtr<FHierarchyElementViewModel> ViewModel = HierarchyMenuContext->MenuHierarchyElements[0];
+	if(ViewModel->IsForHierarchy())
+	{
+		return false;
+	}
+	
+	if(TSharedPtr<FHierarchyElementViewModel> MatchingViewModelInHierarchy = GetHierarchyRootViewModel()->FindViewModelForChild(ViewModel->GetData()->GetPersistentIdentity(), true))
+	{
+		return MatchingViewModelInHierarchy.IsValid();
+	}
+	
+	return false;
+}
+
+const TArray<TSharedPtr<FHierarchyElementViewModel>>& UHierarchyViewModelBase::GetHierarchyItems() const
+{
+	return HierarchyRootViewModel->GetFilteredChildren();
+}
+
+TSharedRef<FHierarchyDragDropOp> UHierarchyViewModelBase::CreateDragDropOp(TSharedRef<FHierarchyElementViewModel> Item)
+{
+	TSharedRef<FHierarchyDragDropOp> DragDropOp = MakeShared<FHierarchyDragDropOp>(Item);
+	DragDropOp->Construct();
+	return DragDropOp;
+}
+
+void UHierarchyViewModelBase::OnGetChildren(TSharedPtr<FHierarchyElementViewModel> Element, TArray<TSharedPtr<FHierarchyElementViewModel>>& OutChildren) const
+{
+	OutChildren.Append(Element->GetFilteredChildren());
+}
+
+void UHierarchyViewModelBase::SetActiveHierarchySection(TSharedPtr<FHierarchySectionViewModel> Section)
+{
+	ActiveHierarchySection = Section;	
+	RefreshHierarchyView(true);
+	OnHierarchySectionActivatedDelegate.ExecuteIfBound(Section);
+}
+
+TSharedPtr<FHierarchySectionViewModel> UHierarchyViewModelBase::GetActiveHierarchySection() const
+{
+	return ActiveHierarchySection.IsValid() ? ActiveHierarchySection.Pin() : nullptr;
+}
+
+UHierarchySection* UHierarchyViewModelBase::GetActiveHierarchySectionData() const
+{
+	return ActiveHierarchySection.IsValid() ? ActiveHierarchySection.Pin()->GetDataMutable<UHierarchySection>() : nullptr;
+}
+
+bool UHierarchyViewModelBase::IsHierarchySectionActive(const UHierarchySection* Section) const
+{
+	return ActiveHierarchySection == nullptr || ActiveHierarchySection.Pin()->GetData() == Section;
+}
+
+FString UHierarchyViewModelBase::OnElementToStringDebug(TSharedPtr<FHierarchyElementViewModel> ElementViewModel) const
+{
+	return ElementViewModel->ToString();	
+}
+
+FHierarchyElementViewModel::~FHierarchyElementViewModel()
+{
+	Children.Empty();
+	FilteredChildren.Empty();
+}
+
+UHierarchyElement* FHierarchyElementViewModel::AddChild(TSubclassOf<UHierarchyElement> NewChildClass, FHierarchyElementIdentity ChildIdentity)
+{
+	UHierarchyElement* NewChild = NewObject<UHierarchyElement>(GetDataMutable(), NewChildClass);
+	NewChild->SetFlags(RF_Transactional);
+	NewChild->Modify();
+	NewChild->SetIdentity(ChildIdentity);
+	GetDataMutable()->GetChildrenMutable().Add(NewChild);
+	
+	SyncViewModelsToData();
+	HierarchyViewModel->OnHierarchyChanged().Broadcast();
+	return NewChild;
+}
+
+void FHierarchyElementViewModel::Tick(float DeltaTime)
+{
+	if(bRenamePending)
+	{
+		RequestRename();
+	}
+}
+
+TStatId FHierarchyElementViewModel::GetStatId() const
+{
+	RETURN_QUICK_DECLARE_CYCLE_STAT(FHierarchyElementViewModel, STATGROUP_Tickables);
+}
+
+void FHierarchyElementViewModel::RefreshChildrenData()
+{
+	TArray<TSharedPtr<FHierarchyElementViewModel>> TmpChildren = Children;
+	for(TSharedPtr<FHierarchyElementViewModel> Child : TmpChildren)
+	{
+		if(Child->RepresentsExternalData() && Child->DoesExternalDataStillExist(HierarchyViewModel->GetRefreshContext()) == false)
+		{
+			Child->Delete();
+		}
+	}
+
+	/** Every item view model can define its own sort order for its children. */
+	SortChildrenData();
+	
+	RefreshChildrenDataInternal();
+
+	/** All remaining children are supposed to exist at this point, as internal data won't be removed by refreshing & external data was cleaned up already.
+	 * This will not call RefreshChildrenData on data that has just been added as no view models exist for these yet.
+	 */
+	for(TSharedPtr<FHierarchyElementViewModel> Child : Children)
+	{
+		Child->RefreshChildrenData();
+	}
+}
+
+void FHierarchyElementViewModel::SyncViewModelsToData()
+{	
+	// this will recursively remove all outdated external data as well as give individual view models the chance to add new data
+	RefreshChildrenData();
+	
+	// now that the data is refreshed, we can sync to the data by recycling view models & creating new ones
+	// old view models will get deleted automatically
+	TArray<TSharedPtr<FHierarchyElementViewModel>> NewChildren;
+	for(UHierarchyElement* Child : Element->GetChildren())
+	{		
+		int32 ViewModelIndex = FindIndexOfChild(Child);
+		// if we couldn't find a view model for a data child, we create it here
+		if(ViewModelIndex == INDEX_NONE)
+		{
+			TSharedPtr<FHierarchyElementViewModel> ChildViewModel = HierarchyViewModel->CreateViewModelForElement(Child, AsShared());
+			if(ensure(ChildViewModel.IsValid()))
+			{
+				ChildViewModel->Initialize();
+				ChildViewModel->SyncViewModelsToData();
+				NewChildren.Add(ChildViewModel);
+			}
+		}
+		// if we could find the view model, we refresh its contained view models and readd it
+		else
+		{
+			Children[ViewModelIndex]->SyncViewModelsToData();
+			NewChildren.Add(Children[ViewModelIndex]);
+		}
+	}
+
+	Children.Empty();
+	Children.Append(NewChildren);
+	
+	for(TSharedPtr<FHierarchyElementViewModel> Child : Children)
+	{
+		Child->OnChildRequestedDeletion().BindSP(this, &FHierarchyElementViewModel::DeleteChild);
+		Child->GetOnSynced().BindSP(this, &FHierarchyElementViewModel::PropagateOnChildSynced);
+	}
+
+	/** Give the view models a chance to further customize the children sync process. */
+	SyncViewModelsToDataInternal();	
+
+	// then we sort the view models according to the data order as this is what will determine widget order created from the view models
+	Children.Sort([this](const TSharedPtr<FHierarchyElementViewModel>& ItemA, const TSharedPtr<FHierarchyElementViewModel>& ItemB)
+		{
+			return FindIndexOfDataChild(ItemA) < FindIndexOfDataChild(ItemB);
+		});
+	
+	// we refresh the filtered children here as well
+	GetFilteredChildren();
+
+	OnSyncedDelegate.ExecuteIfBound();
+}
+
+const TArray<TSharedPtr<FHierarchyElementViewModel>>& FHierarchyElementViewModel::GetFilteredChildren() const
+{
+	FilteredChildren.Empty();
+
+	if(CanHaveChildren())
+	{
+		for(TSharedPtr<FHierarchyElementViewModel> Child : Children)
+		{
+			bool bPassesFilter = true;
+			for(const FOnFilterChild& OnFilterChild : ChildFilters)
+			{
+				bPassesFilter &= OnFilterChild.Execute(Child);
+
+				if(!bPassesFilter)
+				{
+					break;
+				}
+			}
+
+			if(bPassesFilter)
+			{
+				FilteredChildren.Add(Child);
+			}
+		}
+	}
+
+	return FilteredChildren;
+}
+
+void FHierarchyElementViewModel::SortChildrenData() const
+{
+	GetDataMutable()->GetChildrenMutable().StableSort([](const UHierarchyElement& ItemA, const UHierarchyElement& ItemB)
+		{
+			return ItemA.IsA<UHierarchyCategory>() && ItemB.IsA<UHierarchyItem>();
+		});
+}
+
+int32 FHierarchyElementViewModel::GetHierarchyDepth() const
+{
+	if(Parent.IsValid())
+	{
+		return 1 + Parent.Pin()->GetHierarchyDepth();
+	}
+
+	return 0;
+}
+
+void FHierarchyElementViewModel::AddChildFilter(FOnFilterChild InFilterChild)
+{
+	if(ensure(InFilterChild.IsBound()))
+	{
+		ChildFilters.Add(InFilterChild);
+	}
+}
+
+bool FHierarchyElementViewModel::HasParent(TSharedPtr<FHierarchyElementViewModel> ParentCandidate, bool bRecursive) const
+{
+	if(Parent.IsValid())
+	{
+		if(Parent == ParentCandidate)
+		{
+			return true;
+		}
+		else if(bRecursive)
+		{
+			return Parent.Pin()->HasParent(ParentCandidate, bRecursive);
+		}
+	}
+
+	return false;
+}
+
+TSharedRef<FHierarchyElementViewModel> FHierarchyElementViewModel::DuplicateToThis(TSharedPtr<FHierarchyElementViewModel> ItemToDuplicate, int32 InsertIndex)
+{
+	UHierarchyElement* NewItem = Cast<UHierarchyElement>(StaticDuplicateObject(ItemToDuplicate->GetData(), GetDataMutable()));
+	if(InsertIndex == INDEX_NONE)
+	{
+		GetDataMutable()->GetChildrenMutable().Add(NewItem);
+	}
+	else
+	{
+		GetDataMutable()->GetChildrenMutable().Insert(NewItem, InsertIndex);
+	}
+	
+	SyncViewModelsToData();
+
+	HierarchyViewModel->OnHierarchyChanged().Broadcast();
+	TSharedPtr<FHierarchyElementViewModel> ViewModel = FindViewModelForChild(NewItem);
+	return ViewModel.ToSharedRef();
+}
+
+TSharedRef<FHierarchyElementViewModel> FHierarchyElementViewModel::ReparentToThis(TSharedPtr<FHierarchyElementViewModel> ItemToMove, int32 InsertIndex)
+{
+	UHierarchyElement* NewItem = Cast<UHierarchyElement>(StaticDuplicateObject(ItemToMove->GetData(), GetDataMutable()));
+	if(InsertIndex == INDEX_NONE)
+	{
+		GetDataMutable()->GetChildrenMutable().Add(NewItem);
+	}
+	else
+	{
+		GetDataMutable()->GetChildrenMutable().Insert(NewItem, InsertIndex);
+	}
+	
+	ItemToMove->Delete();
+	SyncViewModelsToData();
+	HierarchyViewModel->OnHierarchyChanged().Broadcast();
+	TSharedPtr<FHierarchyElementViewModel> ViewModel = FindViewModelForChild(NewItem);
+	return ViewModel.ToSharedRef();
+}
+
+TSharedPtr<FHierarchyElementViewModel> FHierarchyElementViewModel::FindViewModelForChild(UHierarchyElement* Child, bool bSearchRecursively) const
+{
+	int32 Index = FindIndexOfChild(Child);
+	if(Index != INDEX_NONE)
+	{
+		return Children[Index];
+	}
+
+	if(bSearchRecursively)
+	{
+		for(TSharedPtr<FHierarchyElementViewModel> ChildViewModel : Children)
+		{
+			TSharedPtr<FHierarchyElementViewModel> FoundViewModel = ChildViewModel->FindViewModelForChild(Child, bSearchRecursively);
+
+			if(FoundViewModel.IsValid())
+			{
+				return FoundViewModel;
+			}
+		}
+	}
+
+	return nullptr;
+}
+
+TSharedPtr<FHierarchyElementViewModel> FHierarchyElementViewModel::FindViewModelForChild(FHierarchyElementIdentity ChildIdentity, bool bSearchRecursively) const
+{
+	for(TSharedPtr<FHierarchyElementViewModel> Child : Children)
+	{
+		if(Child->GetData()->GetPersistentIdentity() == ChildIdentity)
+		{
+			return Child;
+		}
+	}
+
+	if(bSearchRecursively)
+	{
+		for(TSharedPtr<FHierarchyElementViewModel> ChildViewModel : Children)
+		{
+			TSharedPtr<FHierarchyElementViewModel> FoundViewModel = ChildViewModel->FindViewModelForChild(ChildIdentity, bSearchRecursively);
+
+			if(FoundViewModel.IsValid())
+			{
+				return FoundViewModel;
+			}
+		}
+	}
+
+	return nullptr;
+}
+
+int32 FHierarchyElementViewModel::FindIndexOfChild(UHierarchyElement* Child) const
+{
+	return Children.FindLastByPredicate([Child](TSharedPtr<FHierarchyElementViewModel> Item)
+	{
+		return Item->GetData() == Child;
+	});
+}
+
+int32 FHierarchyElementViewModel::FindIndexOfDataChild(TSharedPtr<FHierarchyElementViewModel> Child) const
+{
+	return GetData()->GetChildren().Find(Child->GetDataMutable());
+}
+
+int32 FHierarchyElementViewModel::FindIndexOfDataChild(UHierarchyElement* Child) const
+{
+	return GetData()->GetChildren().Find(Child);
+}
+
+void FHierarchyElementViewModel::Delete()
+{
+	OnChildRequestedDeletionDelegate.Execute(AsShared());
+}
+
+void FHierarchyElementViewModel::DeleteChild(TSharedPtr<FHierarchyElementViewModel> Child)
+{	
+	ensure(Child->GetParent().Pin() == AsShared());
+	GetDataMutable()->Modify();
+	GetDataMutable()->GetChildrenMutable().Remove(Child->GetDataMutable());
+	Children.Remove(Child);
+}
+
+TOptional<EItemDropZone> FHierarchyElementViewModel::OnCanRowAcceptDrop(const FDragDropEvent& DragDropEvent, EItemDropZone ItemDropZone, TSharedPtr<FHierarchyElementViewModel> Item)
+{
+	if(TSharedPtr<FHierarchyDragDropOp> DragDropOp = DragDropEvent.GetOperationAs<FHierarchyDragDropOp>())
+	{
+		FCanPerformActionResults Results = CanDropOn(DragDropOp->GetDraggedElement().Pin(), ItemDropZone);
+		DragDropOp->SetDescription(Results.CanPerformMessage);
+		return Results.bCanPerform ? ItemDropZone : TOptional<EItemDropZone>();
+	}
+
+	return TOptional<EItemDropZone>();
+}
+
+FReply FHierarchyElementViewModel::OnDroppedOnRow(const FDragDropEvent& DragDropEvent, EItemDropZone ItemDropZone, TSharedPtr<FHierarchyElementViewModel> Item)
+{
+	if(TSharedPtr<FHierarchyDragDropOp> HierarchyDragDropOp = DragDropEvent.GetOperationAs<FHierarchyDragDropOp>())
+	{
+		OnDroppedOn(HierarchyDragDropOp->GetDraggedElement().Pin(), ItemDropZone);
+		return FReply::Handled();
+	}
+
+	return FReply::Unhandled();
+}
+
+void FHierarchyElementViewModel::OnRowDragLeave(const FDragDropEvent& DragDropEvent)
+{
+	if(TSharedPtr<FHierarchyDragDropOp> HierarchyDragDropOp = DragDropEvent.GetOperationAs<FHierarchyDragDropOp>())
+	{
+		HierarchyDragDropOp->SetDescription(FText::GetEmpty());
+	}
+}
+
+FHierarchyElementViewModel::FCanPerformActionResults FHierarchyElementViewModel::CanDrag()
+{
+	FCanPerformActionResults Results = IsEditableByUser();
+	if(Results.bCanPerform == false)
+	{
+		return Results;
+	}
+
+	return CanDragInternal();
+}
+
+FHierarchyElementViewModel::FCanPerformActionResults FHierarchyElementViewModel::CanDropOnInternal(TSharedPtr<FHierarchyElementViewModel>, EItemDropZone ItemDropZone)
+{
+	return false;
+}
+
+void FHierarchyElementViewModel::PropagateOnChildSynced()
+{
+	OnSyncedDelegate.ExecuteIfBound();
+}
+
+FReply FHierarchyElementViewModel::OnDragDetected(const FGeometry& Geometry, const FPointerEvent& PointerEvent,	bool bIsSource)
+{
+	FCanPerformActionResults CanDragResults = CanDrag();
+	if(CanDragResults == true)
+	{
+		// if the drag is coming from source, we check if any of the hierarchy data already contains that element and we don't start a drag drop in that case
+		if(bIsSource)
+		{
+			TArray<TSharedPtr<FHierarchyElementViewModel>> AllChildren;
+			GetChildrenViewModelsForType<UHierarchyElement, FHierarchyElementViewModel>(AllChildren, true);
+
+			bool bCanDrag = GetHierarchyViewModel()->GetHierarchyRootViewModel()->FindViewModelForChild(GetData()->GetPersistentIdentity(), true) == nullptr;			
+
+			if(bCanDrag)
+			{
+				for(TSharedPtr<FHierarchyElementViewModel>& ChildViewModel : AllChildren)
+				{
+					if(GetHierarchyViewModel()->GetHierarchyRootViewModel()->FindViewModelForChild(ChildViewModel->GetData()->GetPersistentIdentity(), true) != nullptr)
+					{
+						bCanDrag = false;
+						break;
+					}
+				}
+			}
+			
+			if(bCanDrag == false)
+			{
+				return FReply::Unhandled();
+			}
+		}
+		
+		TSharedRef<FHierarchyDragDropOp> HierarchyDragDropOp = HierarchyViewModel->CreateDragDropOp(AsShared());
+		HierarchyDragDropOp->SetFromSourceList(bIsSource);
+
+		return FReply::Handled().BeginDragDrop(HierarchyDragDropOp);			
+	}
+	else
+	{
+		// if we can't drag and have a message, we show it as a slate notification
+		if(CanDragResults.CanPerformMessage.IsEmpty() == false)
+		{
+			FNotificationInfo CantDragInfo(CanDragResults.CanPerformMessage);
+			FSlateNotificationManager::Get().AddNotification(CantDragInfo);
+		}
+	}
+		
+	return FReply::Unhandled();
+}
+
+FHierarchyRootViewModel::~FHierarchyRootViewModel()
+{
+	
+}
+
+void FHierarchyRootViewModel::Initialize()
+{
+	GetOnSynced().BindSP(this, &FHierarchyRootViewModel::PropagateOnSynced);
+}
+
+FHierarchyElementViewModel::FCanPerformActionResults FHierarchyRootViewModel::CanDropOnInternal(TSharedPtr<FHierarchyElementViewModel> DraggedElement, EItemDropZone ItemDropZone)
+{
+	FCanPerformActionResults Results(false);
+	
+	// we only allow drops if some general conditions are fulfilled
+	if(DraggedElement->GetData() != GetData() &&
+		(!DraggedElement->HasParent(AsShared(), false) || ItemDropZone != EItemDropZone::OntoItem)  &&
+		!HasParent(DraggedElement, true))
+	{
+		Results.bCanPerform = 
+			// items can be dropped onto the root directly if the section is set to "All"
+			(DraggedElement->GetData()->IsA<UHierarchyItem>() && HierarchyViewModel->GetActiveHierarchySectionData() == nullptr)
+				||
+			// categories can be dropped onto the root always
+			(DraggedElement->GetData()->IsA<UHierarchyCategory>());
+
+		if(Results.bCanPerform)
+		{
+			if(DraggedElement->IsForHierarchy() == false)
+			{
+				FText Message = LOCTEXT("CanDropSourceItemOnRootDragMessage", "Add {0} to the hierarchy root.");
+				Message = FText::FormatOrdered(Message, FText::FromString(DraggedElement->ToString()));
+				Results.CanPerformMessage = Message;
+			}
+			else
+			{
+				FText Message = LOCTEXT("CanDropHierarchyItemOnRootDragMessage", "Move {0} to the hierarchy root.");
+				Message = FText::FormatOrdered(Message, FText::FromString(DraggedElement->ToString()));
+				Results.CanPerformMessage = Message;
+			}			
+		}
+		else
+		{
+			FText Message = LOCTEXT("CantDropHierarchyItemOnRootDragMessage", "Can not add {0} here. Please add it to a category!");
+			Message = FText::FormatOrdered(Message, FText::FromString(DraggedElement->ToString()));
+			Results.CanPerformMessage = Message;
+		}
+	}
+	
+	return Results;
+}
+
+void FHierarchyRootViewModel::OnDroppedOnInternal(TSharedPtr<FHierarchyElementViewModel> DroppedItem, EItemDropZone ItemDropZone)
+{
+	FScopedTransaction Transaction(LOCTEXT("Transaction_OnDropOnRoot", "Dropped item on root"));
+	HierarchyViewModel->GetHierarchyRoot()->Modify();
+
+	if(DroppedItem->GetDataMutable()->IsA<UHierarchyItem>() || DroppedItem->GetDataMutable()->IsA<UHierarchyCategory>())
+	{
+		TSharedPtr<FHierarchyElementViewModel> NewViewModel;
+		// we duplicate the item if the dragged item is from source
+		if(DroppedItem->IsForHierarchy() == false)
+		{
+			NewViewModel = DuplicateToThis(DroppedItem);
+		}
+		else
+		{
+			NewViewModel = ReparentToThis(DroppedItem);
+		}
+
+		if(UHierarchyCategory* AsCategory = Cast<UHierarchyCategory>(NewViewModel->GetDataMutable()))
+		{
+			AsCategory->SetSection(HierarchyViewModel->GetActiveHierarchySectionData());
+		}
+
+		HierarchyViewModel->RefreshHierarchyView();
+	}
+}
+
+TSharedPtr<FHierarchySectionViewModel> FHierarchyRootViewModel::AddSection()
+{
+	FScopedTransaction ScopedTransaction(LOCTEXT("NewSectionAdded","Added Section"));
+	HierarchyViewModel->GetHierarchyRoot()->Modify();
+	
+	UHierarchySection* SectionData = GetDataMutable<UHierarchyRoot>()->AddSection(LOCTEXT("HierarchyEditorDefaultNewSectionName", "Section"), 0, HierarchyViewModel->GetSectionDataClass());
+	SectionData->Modify();
+	TSharedPtr<FHierarchyElementViewModel> ViewModel = HierarchyViewModel->CreateViewModelForElement(SectionData, StaticCastSharedRef<FHierarchyRootViewModel>(AsShared()));
+	TSharedPtr<FHierarchySectionViewModel> SectionViewModel = StaticCastSharedPtr<FHierarchySectionViewModel>(ViewModel);
+	if(!ensureMsgf(SectionViewModel.IsValid(), TEXT("Make sure that CreateViewModelForData creates a FHierarchySectionViewModel (or derived) for UHierarchySection elements")))
+	{
+		return nullptr;	
+	}
+	
+	SectionViewModels.Add(SectionViewModel);
+	SyncViewModelsToData();
+	HierarchyViewModel->SetActiveHierarchySection(SectionViewModel);
+
+	OnSectionAddedDelegate.ExecuteIfBound(SectionViewModel);
+	OnSectionsChangedDelegate.ExecuteIfBound();
+	return SectionViewModel;
+}
+
+void FHierarchyRootViewModel::DeleteSection(TSharedPtr<FHierarchyElementViewModel> InSectionViewModel)
+{
+	TSharedPtr<FHierarchySectionViewModel> SectionViewModel = StaticCastSharedPtr<FHierarchySectionViewModel>(InSectionViewModel);
+	GetDataMutable<UHierarchyRoot>()->GetSectionDataMutable().Remove(SectionViewModel->GetDataMutable<UHierarchySection>());
+	SectionViewModels.Remove(SectionViewModel);
+
+	OnSectionDeletedDelegate.ExecuteIfBound(SectionViewModel);
+	OnSectionsChangedDelegate.ExecuteIfBound();
+}
+
+void FHierarchyRootViewModel::PropagateOnSynced()
+{
+	OnSyncPropagatedDelegate.ExecuteIfBound();
+}
+
+void FHierarchyRootViewModel::SyncViewModelsToDataInternal()
+{
+	const UHierarchyRoot* RootData = GetData<UHierarchyRoot>();
+
+	TArray<TSharedPtr<FHierarchySectionViewModel>> NewSectionViewModels;
+	TArray<TSharedPtr<FHierarchySectionViewModel>> SectionViewModelsToDelete;
+	
+	for(TSharedPtr<FHierarchySectionViewModel> SectionViewModel : SectionViewModels)
+	{
+		if(!RootData->GetSectionData().Contains(SectionViewModel->GetData()))
+		{
+			SectionViewModelsToDelete.Add(SectionViewModel);
+		}
+	}
+
+	for (TSharedPtr<FHierarchySectionViewModel> SectionViewModel : SectionViewModelsToDelete)
+	{
+		SectionViewModel->Delete();
+	}
+	
+	for(UHierarchySection* Section : RootData->GetSectionData())
+	{
+		TSharedPtr<FHierarchySectionViewModel>* SectionViewModelPtr = SectionViewModels.FindByPredicate([Section](TSharedPtr<FHierarchySectionViewModel> SectionViewModel)
+		{
+			return SectionViewModel->GetData() == Section;
+		});
+
+		TSharedPtr<FHierarchySectionViewModel> SectionViewModel = nullptr;
+
+		if(SectionViewModelPtr)
+		{
+			SectionViewModel = *SectionViewModelPtr;
+		}
+
+		if(SectionViewModel == nullptr)
+		{
+			SectionViewModel = MakeShared<FHierarchySectionViewModel>(Section, StaticCastSharedRef<FHierarchyRootViewModel>(AsShared()), HierarchyViewModel);
+			SectionViewModel->SyncViewModelsToData();;
+		}
+		
+		NewSectionViewModels.Add(SectionViewModel);
+	}
+
+	SectionViewModels.Empty();
+	SectionViewModels.Append(NewSectionViewModels);
+
+	for(TSharedPtr<FHierarchySectionViewModel> SectionViewModel : SectionViewModels)
+	{
+		SectionViewModel->OnChildRequestedDeletion().BindSP(this, &FHierarchyRootViewModel::DeleteSection);
+	}
+
+	SectionViewModels.Sort([this](const TSharedPtr<FHierarchySectionViewModel>& ItemA, const TSharedPtr<FHierarchySectionViewModel>& ItemB)
+		{
+			return
+			GetDataMutable<UHierarchyRoot>()->GetSectionData().Find(Cast<UHierarchySection>(ItemA->GetDataMutable()))
+				<
+			GetDataMutable<UHierarchyRoot>()->GetSectionData().Find(Cast<UHierarchySection>(ItemB->GetDataMutable())); 
+		});
+}
+
+FHierarchyElementViewModel::FCanPerformActionResults FHierarchySectionViewModel::CanDropOnInternal(TSharedPtr<FHierarchyElementViewModel> DraggedElement, EItemDropZone ItemDropZone)
+{
+	FCanPerformActionResults Results(false);
+	// we don't allow dropping onto source sections and we don't specify a message as the sections aren't going to light up as valid drop targets
+	if(IsForHierarchy() == false)
+	{
+		return false;
+	}
+	
+	if(const UHierarchyCategory* Category = Cast<UHierarchyCategory>(DraggedElement->GetData()))
+	{
+		if(ItemDropZone == EItemDropZone::OntoItem)
+		{
+			FText Message = LOCTEXT("DropCategoryOnSectionDragMessage", "Add {0} to section {1}");
+			Message = FText::FormatOrdered(Message, FText::FromString(DraggedElement->ToString()), FText::FromString(ToString()));
+			
+			Results.CanPerformMessage = Message;
+			Results.bCanPerform = GetData() != Category->GetSection();
+		}
+	}
+	else if(UHierarchySection* Section = Cast<UHierarchySection>(DraggedElement->GetDataMutable()))
+	{
+		const bool bSameSection = GetData() == Section;
+
+		if(ItemDropZone == EItemDropZone::OntoItem && bSameSection == false)
+		{
+			FText Message = LOCTEXT("CantDropSectionOnSectionDragMessage", "Can't drop section on a section.");
+			Results.CanPerformMessage = Message;
+			Results.bCanPerform = false;
+			return Results;
+		}
+		
+		int32 DraggedElementIndex = GetHierarchyViewModel()->GetHierarchyRoot()->GetSectionIndex(Section);
+		int32 InsertionIndex = GetHierarchyViewModel()->GetHierarchyRoot()->GetSectionIndex(GetDataMutable<UHierarchySection>());
+		// we add 1 to the insertion index if it's below an item because we either want to insert at the current index to place the item above, or at current+1 for below
+		InsertionIndex += ItemDropZone == EItemDropZone::AboveItem ? -1 : 1;
+
+		Results.bCanPerform = !bSameSection && DraggedElementIndex != InsertionIndex;
+
+		if(Results.bCanPerform)
+		{
+			if(ItemDropZone == EItemDropZone::AboveItem)
+			{
+				FText Message = LOCTEXT("MoveSectionLeftDragMessage", "Move section {0} to the left");
+				Message = FText::FormatOrdered(Message, FText::FromString(DraggedElement->ToString()));
+				Results.CanPerformMessage = Message;
+			}
+			else if(ItemDropZone == EItemDropZone::BelowItem)
+			{
+				FText Message = LOCTEXT("MoveSectionRightDragMessage", "Move section {0} to the right");
+				Message = FText::FormatOrdered(Message, FText::FromString(DraggedElement->ToString()));
+				Results.CanPerformMessage = Message;
+			}
+		}
+	}
+	else if(UHierarchyItem* Item = Cast<UHierarchyItem>(DraggedElement->GetDataMutable()))
+	{
+		FText Message = LOCTEXT("CantDropItemOnSectionDragMessage", "Can't drop items onto sections. Please drag a category onto section {0}");
+		Message = FText::FormatOrdered(Message, FText::FromString(ToString()));
+		Results.bCanPerform = false;
+		Results.CanPerformMessage = Message;
+	}
+
+	return Results;
+}
+
+void FHierarchySectionViewModel::OnDroppedOnInternal(TSharedPtr<FHierarchyElementViewModel> DroppedItem, EItemDropZone ItemDropZone)
+{
+	if(DroppedItem->GetData()->IsA<UHierarchySection>())
+	{
+		UHierarchySection* DraggedSectionData = DroppedItem->GetDataMutable<UHierarchySection>();
+
+		int32 IndexOfThis = HierarchyViewModel->GetHierarchyRoot()->GetSectionData().Find(GetDataMutable<UHierarchySection>());
+		int32 DraggedSectionIndex = HierarchyViewModel->GetHierarchyRoot()->GetSectionData().Find(DraggedSectionData);
+
+		TArray<TObjectPtr<UHierarchySection>>& SectionData = HierarchyViewModel->GetHierarchyRoot()->GetSectionDataMutable();
+		int32 Count = SectionData.Num();
+
+		bool bDropSucceeded = false;
+		// above constitutes to the left here
+		if(ItemDropZone == EItemDropZone::AboveItem)
+		{
+			SectionData.RemoveAt(DraggedSectionIndex);
+			SectionData.Insert(DraggedSectionData, FMath::Max(IndexOfThis, 0));
+
+			bDropSucceeded = true;
+		}
+		else if(ItemDropZone == EItemDropZone::BelowItem)
+		{
+			SectionData.RemoveAt(DraggedSectionIndex);
+
+			if(IndexOfThis + 1 > SectionData.Num())
+			{
+				SectionData.Add(DraggedSectionData);
+			}
+			else
+			{
+				SectionData.Insert(DraggedSectionData, FMath::Min(IndexOfThis+1, Count));
+			}
+
+			bDropSucceeded = true;
+
+		}
+
+		if(bDropSucceeded)
+		{
+			HierarchyViewModel->ForceFullRefresh();
+			HierarchyViewModel->OnHierarchyChanged().Broadcast();
+		}
+	}
+	else
+	{		
+		if(UHierarchyCategory* HierarchyCategory = DroppedItem->GetDataMutable<UHierarchyCategory>())
+		{
+			FScopedTransaction Transaction(LOCTEXT("Transaction_OnSectionDrop", "Moved category to section"));
+			HierarchyViewModel->GetHierarchyRoot()->Modify();
+			
+			HierarchyCategory->SetSection(GetDataMutable<UHierarchySection>());
+
+			// we null out any sections for all contained categories
+			TArray<UHierarchyCategory*> AllChildCategories;
+			HierarchyCategory->GetChildrenOfType<UHierarchyCategory>(AllChildCategories, true);
+			for(UHierarchyCategory* ChildCategory : AllChildCategories)
+			{
+				ChildCategory->SetSection(nullptr);
+			}
+
+			// we only need to reparent if the parent isn't already the root. This stops unnecessary reordering
+			if(DroppedItem->GetParent() != HierarchyViewModel->GetHierarchyRootViewModel())
+			{
+				HierarchyViewModel->GetHierarchyRootViewModel()->ReparentToThis(DroppedItem);
+			}
+			
+			HierarchyViewModel->RefreshHierarchyView();
+			HierarchyViewModel->OnHierarchyChanged().Broadcast();
+		}
+	}
+}
+
+void FHierarchySectionViewModel::FinalizeInternal()
+{
+	if(HierarchyViewModel->GetActiveHierarchySection() == AsShared())
+	{
+		HierarchyViewModel->SetActiveHierarchySection(nullptr);
+	}
+
+	// we make sure to reset all categories' section entry that were referencing this section
+	TArray<UHierarchyCategory*> AllCategories;
+	HierarchyViewModel->GetHierarchyRoot()->GetChildrenOfType<UHierarchyCategory>(AllCategories, true);
+
+	for(UHierarchyCategory* Category : AllCategories)
+	{
+		if(Category->GetSection() == GetData())
+		{
+			Category->SetSection(nullptr);
+		}
+	}
+}
+
+::FHierarchyElementViewModel::FCanPerformActionResults FHierarchyItemViewModel::CanDropOnInternal(TSharedPtr<FHierarchyElementViewModel> DraggedElement, EItemDropZone ItemDropZone)
+{
+	bool bAllowDrop = false;
+	
+	TSharedPtr<FHierarchyElementViewModel> SourceDropItem = DraggedElement;
+	TSharedPtr<FHierarchyElementViewModel> TargetDropItem = AsShared();
+
+	// we only allow drops if some general conditions are fulfilled
+	if(SourceDropItem->GetData() != TargetDropItem->GetData() &&
+		(!SourceDropItem->HasParent(TargetDropItem, false) || ItemDropZone != EItemDropZone::OntoItem)  &&
+		!TargetDropItem->HasParent(SourceDropItem, true))
+	{
+		// items can be generally be dropped above/below other items
+		bAllowDrop = (SourceDropItem->GetData()->IsA<UHierarchyItem>() && ItemDropZone != EItemDropZone::OntoItem);
+	}
+
+	return bAllowDrop;
+}
+
+void FHierarchyItemViewModel::OnDroppedOnInternal(TSharedPtr<FHierarchyElementViewModel> DroppedItem, EItemDropZone ItemDropZone)
+{
+	FScopedTransaction Transaction(LOCTEXT("Transaction_MovedItem", "Moved an item in the hierarchy"));
+	HierarchyViewModel->GetHierarchyRoot()->Modify();
+	
+	bool bDropSucceeded = false;
+	if(ItemDropZone == EItemDropZone::AboveItem)
+	{
+		int32 IndexOfThis = Parent.Pin()->FindIndexOfDataChild(AsShared());
+
+		if(DroppedItem->IsForHierarchy() == false)
+		{
+			Parent.Pin()->DuplicateToThis(DroppedItem, FMath::Max(IndexOfThis, 0));
+		}
+		else
+		{
+			Parent.Pin()->ReparentToThis(DroppedItem, FMath::Max(IndexOfThis, 0));
+		}
+
+		bDropSucceeded = true;
+	}
+	else if(ItemDropZone == EItemDropZone::BelowItem)
+	{
+		int32 IndexOfThis = Parent.Pin()->FindIndexOfDataChild(AsShared());
+		
+		if(DroppedItem->IsForHierarchy() == false)
+		{
+			Parent.Pin()->DuplicateToThis(DroppedItem, FMath::Min(IndexOfThis+1, Parent.Pin()->GetChildren().Num()));
+		}
+		else
+		{
+			Parent.Pin()->ReparentToThis(DroppedItem, FMath::Min(IndexOfThis+1, Parent.Pin()->GetChildren().Num()));
+		}
+
+		bDropSucceeded = true;
+	}
+
+	if(bDropSucceeded)
+	{
+		HierarchyViewModel->RefreshHierarchyView();
+		HierarchyViewModel->RefreshSourceView();
+	}
+	else
+	{
+		Transaction.Cancel();
+	}
+}
+
+void FHierarchyCategoryViewModel::OnDroppedOnInternal(TSharedPtr<FHierarchyElementViewModel> DroppedItem, EItemDropZone ItemDropZone)
+{
+	FScopedTransaction Transaction(LOCTEXT("Transaction_OnCategoryDrop", "Dropped item on/above/below category"));
+	HierarchyViewModel->GetHierarchyRoot()->Modify();
+	
+	if(UHierarchyCategory* Category = DroppedItem->GetDataMutable<UHierarchyCategory>())
+	{
+		// if we are dragging a category above/below another category and the new parent is going to be the root, we update its section to the active section
+		if(ItemDropZone != EItemDropZone::OntoItem)
+		{
+			if(Parent.IsValid() && Parent == HierarchyViewModel->GetHierarchyRootViewModel())
+			{
+				Category->SetSection(HierarchyViewModel->GetActiveHierarchySectionData());
+
+				// we null out any sections for all contained categories
+				TArray<UHierarchyCategory*> AllChildCategories;
+				Category->GetChildrenOfType<UHierarchyCategory>(AllChildCategories, true);
+				for(UHierarchyCategory* ChildCategory : AllChildCategories)
+				{
+					ChildCategory->SetSection(nullptr);
+				}
+			}				
+		}
+		// if we are dragging a category onto another category, we null out its section instead
+		else
+		{
+			Category->SetSection(nullptr);
+
+			// we null out any sections for all contained categories
+			TArray<UHierarchyCategory*> AllChildCategories;
+			Category->GetChildrenOfType<UHierarchyCategory>(AllChildCategories, true);
+			for(UHierarchyCategory* ChildCategory : AllChildCategories)
+			{
+				ChildCategory->SetSection(nullptr);
+			}
+		}			
+	}
+
+	// the actual moving of the item happens here
+	if(ItemDropZone == EItemDropZone::OntoItem)
+	{
+		if(DroppedItem->IsForHierarchy() == false)
+		{
+			DuplicateToThis(DroppedItem);
+		}
+		else
+		{
+			ReparentToThis(DroppedItem);
+		}
+	}
+	else if(ItemDropZone == EItemDropZone::AboveItem)
+	{
+		int32 IndexOfThis = Parent.Pin()->FindIndexOfDataChild(AsShared());
+		if(DroppedItem->IsForHierarchy() == false)
+		{
+			Parent.Pin()->DuplicateToThis(DroppedItem, FMath::Max(IndexOfThis, 0));
+		}
+		else
+		{				
+			Parent.Pin()->ReparentToThis(DroppedItem, FMath::Max(IndexOfThis, 0));
+		}
+	}
+	else if(ItemDropZone == EItemDropZone::BelowItem)
+	{
+		int32 IndexOfThis = Parent.Pin()->FindIndexOfDataChild(AsShared());
+		if(DroppedItem->IsForHierarchy() == false)
+		{
+			Parent.Pin()->DuplicateToThis(DroppedItem, FMath::Min(IndexOfThis+1, Parent.Pin()->GetChildren().Num()));
+		}
+		else
+		{
+			Parent.Pin()->ReparentToThis(DroppedItem, FMath::Min(IndexOfThis+1, Parent.Pin()->GetChildren().Num()));
+		}
+	}
+}
+
+void UHierarchyViewModelBase::AddCategory(TSharedPtr<FHierarchyElementViewModel> CategoryParent) const
+{
+	// If no category parent was specified, we add it to the root
+	if(CategoryParent == nullptr)
+	{
+		CategoryParent = GetHierarchyRootViewModel();	
+	}
+	
+	int32 HierarchyDepth = CategoryParent->GetHierarchyDepth();
+	if(HierarchyDepth > 15)
+	{
+		FNotificationInfo Info(LOCTEXT("TooManyNestedCategoriesToastText", "We currently only allow a hierarchy depth of 15."));
+		Info.ExpireDuration = 4.f;
+		FSlateNotificationManager::Get().AddNotification(Info);
+		return;
+	}
+
+	FText TransactionText = FText::FormatOrdered(LOCTEXT("Transaction_AddedItem", "Added new {0} to hierarchy"), FText::FromString(GetCategoryDataClass()->GetName()));
+	FScopedTransaction Transaction(TransactionText);
+	GetHierarchyRoot()->Modify();
+
+	UClass* CategoryClass = GetCategoryDataClass();
+	UHierarchyCategory* Category = Cast<UHierarchyCategory>(CategoryParent->AddChild(CategoryClass, FHierarchyElementIdentity()));
+	
+	TSharedPtr<FHierarchyElementViewModel> ViewModel = CategoryParent->FindViewModelForChild(Category, false);
+	if(ensureMsgf(ViewModel.IsValid(), TEXT("Could not find view model for new category of type '%s'. Please ensure your 'CreateViewModelForData' function creates a view model."), *CategoryClass->GetName()))
+	{		
+		TArray<UHierarchyCategory*> SiblingCategories;
+		Category->GetTypedOuter<UHierarchyElement>()->GetChildrenOfType<UHierarchyCategory>(SiblingCategories);
+		
+		TSet<FName> CategoryNames;
+		for(const auto& SiblingCategory : SiblingCategories)
+		{
+			CategoryNames.Add(SiblingCategory->GetCategoryName());
+		}
+
+		Category->SetCategoryName(UE::HierarchyEditor::GetUniqueName(FName("New Category"), CategoryNames));
+		
+		// we only set the section property if the current section isn't set to "All"
+		Category->SetSection(GetActiveHierarchySectionData());
+		
+		RefreshHierarchyView();
+
+		OnElementAddedDelegate.ExecuteIfBound(ViewModel);
+	}
+}
+
+void UHierarchyViewModelBase::AddSection() const
+{
+	TSharedPtr<FHierarchySectionViewModel> SectionViewModel = GetHierarchyRootViewModel()->AddSection();	
+	OnElementAddedDelegate.ExecuteIfBound(SectionViewModel);
+	OnHierarchyChangedDelegate.Broadcast();
+}
+
+void UHierarchyViewModelBase::GenerateDynamicContextMenu(UToolMenu* ToolMenu) const
+{
+	UHierarchyMenuContext* HierarchyMenuContext = ToolMenu->FindContext<UHierarchyMenuContext>();
+
+	if(HierarchyMenuContext == nullptr)
+	{
+		return;
+	}
+	
+	ToolMenu->AddMenuEntry("Dynamic", FToolMenuEntry::InitMenuEntryWithCommandList(FHierarchyEditorCommands::Get().FindInHierarchy, Commands, TAttribute<FText>(), TAttribute<FText>(), FSlateIcon(FAppStyle::GetAppStyleSetName(), "Icons.Find")));
+	ToolMenu->AddMenuEntry("Dynamic", FToolMenuEntry::InitMenuEntryWithCommandList(FGenericCommands::Get().Rename, Commands));
+	ToolMenu->AddMenuEntry("Dynamic", FToolMenuEntry::InitMenuEntryWithCommandList(FGenericCommands::Get().Delete, Commands));
+
+	if(HierarchyMenuContext->MenuHierarchyElements.Num() == 1)
+	{
+		HierarchyMenuContext->MenuHierarchyElements[0]->AppendDynamicContextMenuForSingleElement(ToolMenu);
+	}
+}
+
+UHierarchyElement* UHierarchyViewModelBase::AddElementUnderRoot(TSubclassOf<UHierarchyElement> NewChildClass, FHierarchyElementIdentity ChildIdentity)
+{
+	FScopedTransaction Transaction(LOCTEXT("Transaction_AddItem", "Add hierarchy item"));
+	HierarchyRoot->Modify();
+	return GetHierarchyRootViewModel()->AddChild(NewChildClass, ChildIdentity);
+}
+
+void UHierarchyViewModelBase::DeleteElementWithIdentity(FHierarchyElementIdentity Identity)
+{
+	FScopedTransaction Transaction(LOCTEXT("Transaction_DeleteItem", "Deleted hierarchy item"));
+	HierarchyRoot->Modify();
+
+	bool bItemDeleted = false;
+	if(TSharedPtr<FHierarchyElementViewModel> ViewModel = HierarchyRootViewModel->FindViewModelForChild(Identity, true))
+	{
+		if(ViewModel->CanDelete())
+		{
+			ViewModel->Delete();
+			bItemDeleted = true;
+		}
+	}
+	
+	TArray<TSharedPtr<FHierarchySectionViewModel>> SectionViewModels = HierarchyRootViewModel->GetSectionViewModels();
+	for(TSharedPtr<FHierarchySectionViewModel> SectionViewModel : SectionViewModels)
+	{
+		if(SectionViewModel->GetData()->GetPersistentIdentity() == Identity && SectionViewModel->CanDelete())
+		{
+			SectionViewModel->Delete();
+			bItemDeleted = true;
+		}
+	}
+
+	if(bItemDeleted)
+	{
+		HierarchyRootViewModel->SyncViewModelsToData();
+		OnHierarchyChangedDelegate.Broadcast();
+	}
+}
+
+void UHierarchyViewModelBase::DeleteElementsWithIdentities(TArray<FHierarchyElementIdentity> Identities)
+{
+	TArray<TSharedPtr<FHierarchyElementViewModel>> ViewModelsToDelete;
+	
+	for(FHierarchyElementIdentity& Identity : Identities)
+	{
+		if(TSharedPtr<FHierarchyElementViewModel> ViewModel = HierarchyRootViewModel->FindViewModelForChild(Identity, true))
+		{
+			ViewModelsToDelete.Add(ViewModel);
+		}
+
+		TArray<TSharedPtr<FHierarchySectionViewModel>> SectionViewModels = HierarchyRootViewModel->GetSectionViewModels();
+		for(TSharedPtr<FHierarchySectionViewModel> SectionViewModel : SectionViewModels)
+		{
+			if(SectionViewModel->GetData()->GetPersistentIdentity() == Identity)
+			{
+				ViewModelsToDelete.Add(SectionViewModel);
+			}
+		}
+	}
+
+	DeleteElements(ViewModelsToDelete);
+}
+
+void UHierarchyViewModelBase::DeleteElements(TArray<TSharedPtr<FHierarchyElementViewModel>> ViewModels) const
+{
+	FScopedTransaction Transaction(LOCTEXT("Transaction_DeleteHierarchyElements", "Deleted hierarchy elements"));
+	HierarchyRoot->Modify();
+	
+	bool bAnyItemsDeleted = false;
+	for(TSharedPtr<FHierarchyElementViewModel> ViewModel : ViewModels)
+	{
+		if(ViewModel->CanDelete())
+		{
+			ViewModel->Delete();
+			bAnyItemsDeleted = true;
+		}
+	}
+
+	if(bAnyItemsDeleted)
+	{
+		HierarchyRootViewModel->SyncViewModelsToData();
+		OnHierarchyChangedDelegate.Broadcast();
+	}
+}
+
+void UHierarchyViewModelBase::NavigateToElementInHierarchy(const FHierarchyElementIdentity& HierarchyIdentity) const
+{
+	OnNavigateToElementIdentityInHierarchyRequestedDelegate.ExecuteIfBound(HierarchyIdentity);
+}
+
+void UHierarchyViewModelBase::NavigateToElementInHierarchy(const TSharedRef<FHierarchyElementViewModel> HierarchyElement) const
+{
+	OnNavigateToElementInHierarchyRequestedDelegate.ExecuteIfBound(HierarchyElement);
+}
+
+FHierarchyDragDropOp::FHierarchyDragDropOp(TSharedPtr<FHierarchyElementViewModel> InDraggedElementViewModel) : DraggedElement(InDraggedElementViewModel)
+{
+	SetLabel(DraggedElement.Pin()->ToStringAsText());
+}
+
+TSharedPtr<SWidget> FHierarchyDragDropOp::GetDefaultDecorator() const
+{
+	TSharedRef<SWidget> CustomDecorator = CreateCustomDecorator();
+
+	SVerticalBox::FSlot* CustomSlot;
+	TSharedPtr<SWidget> Decorator = SNew(SToolTip)
+	[
+		SNew(SVerticalBox)
+		+ SVerticalBox::Slot()
+		.Expose(CustomSlot).
+		AutoHeight()
+		+ SVerticalBox::Slot()
+		.AutoHeight()
+		.Padding(2.f)
+		[
+			SNew(STextBlock)
+			.Text(this, &FHierarchyDragDropOp::GetLabel)
+			.TextStyle(&FAppStyle::GetWidgetStyle<FTextBlockStyle>("NormalText.Important"))
+			.Visibility_Lambda([this, CustomDecorator]()
+			{
+				return GetLabel().IsEmpty() || CustomDecorator != SNullWidget::NullWidget ? EVisibility::Collapsed : EVisibility::Visible;
+			})
+		]
+		+ SVerticalBox::Slot()
+		.AutoHeight()
+		.Padding(2.f)
+		[
+			SNew(STextBlock)
+			.Text(this, &FHierarchyDragDropOp::GetDescription)
+			.TextStyle(&FAppStyle::GetWidgetStyle<FTextBlockStyle>("NormalText"))
+			.Visibility_Lambda([this]()
+			{
+				return GetDescription().IsEmpty() ? EVisibility::Collapsed : EVisibility::Visible;
+			})
+		]
+	];
+
+	if(CustomDecorator != SNullWidget::NullWidget)
+	{
+		CustomSlot->AttachWidget(CustomDecorator);
+	}
+
+	return Decorator;
+}
+
+TSharedRef<SWidget> FSectionDragDropOp::CreateCustomDecorator() const
+{
+	return SNew(SCheckBox)
+		.Visibility(EVisibility::HitTestInvisible)
+		.Style(FAppStyle::Get(), "DetailsView.SectionButton")
+		.IsChecked(ECheckBoxState::Unchecked)
+		[
+			SNew(SInlineEditableTextBlock)
+			.Text(GetDraggedSection().Pin()->GetSectionNameAsText())
+		];
+}
+
+#undef LOCTEXT_NAMESPACE
