New Tool: Hierarchy Editor

diff --git a/Engine/Source/Editor/HierarchyEditor/Private/Widgets/SHierarchyEditor.cpp b/Engine/Source/Editor/HierarchyEditor/Private/Widgets/SHierarchyEditor.cpp
new file mode 100644
--- /dev/null
+++ b/Engine/Source/Editor/HierarchyEditor/Private/Widgets/SHierarchyEditor.cpp
@@ -0,0 +1,1765 @@
+// Copyright Epic Games, Inc. All Rights Reserved.
+
+#include "Widgets/SHierarchyEditor.h"
+
+#include "HierarchyEditorCommands.h"
+#include "UObject/Package.h"
+#include "Framework/Commands/GenericCommands.h"
+#include "Styling/StyleColors.h"
+#include "Styling/AppStyle.h"
+#include "HierarchyViewModelBase.h"
+#include "Widgets/Layout/SSeparator.h"
+#include "Widgets/Layout/SWrapBox.h"
+#include "Widgets/Input/SCheckBox.h"
+#include "Widgets/Input/SButton.h"
+#include "Widgets/Images/SImage.h"
+#include "Widgets/Text/SInlineEditableTextBlock.h"
+#include "Widgets/Text/STextBlock.h"
+#include "ScopedTransaction.h"
+#include "Modules/ModuleManager.h"
+#include "PropertyEditorModule.h"
+#include "Framework/Application/MenuStack.h"
+#include "Framework/Application/SlateApplication.h"
+#include "ToolMenus.h"
+#include "HierarchyEditorStyle.h"
+#include "Widgets/Input/SSearchBox.h"
+
+#define LOCTEXT_NAMESPACE "HierarchyEditor"
+
+TSharedRef<SWidget> SummonContextMenu(TArray<TSharedPtr<FHierarchyElementViewModel>> MenuHierarchyElements)
+{
+	UHierarchyMenuContext* MenuContextObject = NewObject<UHierarchyMenuContext>();
+	MenuContextObject->MenuHierarchyElements = MenuHierarchyElements;
+
+	FToolMenuContext MenuContext(MenuContextObject);
+	TWeakObjectPtr<UHierarchyViewModelBase> ViewModel = MenuHierarchyElements[0]->GetHierarchyViewModel();
+	MenuContext.AppendCommandList(ViewModel->GetCommands());
+	
+	TSharedRef<SWidget> MenuWidget = UToolMenus::Get()->GenerateWidget("HierarchyMenu", MenuContext);
+	return MenuWidget;
+}
+
+void SHierarchyCategory::Construct(const FArguments& InArgs,	TSharedPtr<FHierarchyCategoryViewModel> InCategoryViewModel)
+{
+	CategoryViewModel = InCategoryViewModel;
+	
+	InCategoryViewModel->GetOnRequestRename().BindSP(this, &SHierarchyCategory::EnterEditingMode);
+
+	const UHierarchyCategory* CategoryData = InCategoryViewModel->GetData<UHierarchyCategory>();
+	SetToolTipText(TAttribute<FText>::CreateLambda([CategoryData]
+	{
+		return CategoryData->GetTooltip();
+	}));
+	
+	ChildSlot
+	[
+		SAssignNew(InlineEditableTextBlock, SInlineEditableTextBlock)
+		.Style(&FHierarchyEditorStyle::Get().GetWidgetStyle<FInlineEditableTextBlockStyle>("HierarchyEditor.CategoryTextStyle"))
+		.Text(this, &SHierarchyCategory::GetCategoryText)
+		.OnTextCommitted(this, &SHierarchyCategory::OnRenameCategory)
+		.OnVerifyTextChanged(this, &SHierarchyCategory::OnVerifyCategoryRename)
+		.IsSelected(InArgs._IsSelected)
+	];
+}
+
+void SHierarchyCategory::EnterEditingMode() const
+{
+	if(CategoryViewModel.Pin()->CanRename())
+	{
+		InlineEditableTextBlock->EnterEditingMode();
+	}
+}
+
+bool SHierarchyCategory::OnVerifyCategoryRename(const FText& NewName, FText& OutTooltip) const
+{
+	TArray<TSharedPtr<FHierarchyCategoryViewModel>> SiblingCategoryViewModels;
+	CategoryViewModel.Pin()->GetParent().Pin()->GetChildrenViewModelsForType<UHierarchyCategory, FHierarchyCategoryViewModel>(SiblingCategoryViewModels);
+
+	if(GetCategoryText().ToString() != NewName.ToString())
+	{
+		TSet<FString> CategoryNames;
+		for(const auto& SiblingCategoryViewModel : SiblingCategoryViewModels)
+		{
+			CategoryNames.Add(Cast<UHierarchyCategory>(SiblingCategoryViewModel->GetDataMutable())->GetCategoryName().ToString());
+		}
+
+		if(CategoryNames.Contains(NewName.ToString()))
+		{
+			OutTooltip = LOCTEXT("HierarchyCategoryCantRename_DuplicateOnLayer", "Another category of the same name already exists!");
+			return false;
+		}
+	}
+
+	return true;
+}
+
+FText SHierarchyCategory::GetCategoryText() const
+{
+	return FText::FromString(CategoryViewModel.Pin()->ToString());
+}
+
+void SHierarchyCategory::OnRenameCategory(const FText& NewText, ETextCommit::Type) const
+{
+	FScopedTransaction Transaction(LOCTEXT("Transaction_Rename_Category", "Renamed hierarchy category"));
+	CategoryViewModel.Pin()->GetHierarchyViewModel()->GetHierarchyRoot()->Modify();
+	
+	CategoryViewModel.Pin()->Rename(FName(NewText.ToString()));
+}
+
+void SHierarchySection::Construct(const FArguments& InArgs, TSharedPtr<FHierarchySectionViewModel> InSection, TWeakObjectPtr<UHierarchyViewModelBase> InHierarchyViewModel)
+{
+	SectionViewModel = InSection;
+	HierarchyViewModel = InHierarchyViewModel;
+
+	IsSectionActive = InArgs._IsSectionActive;
+	OnSectionActivatedDelegate = InArgs._OnSectionActivated;
+	bForbidDropOn = InArgs._bForbidDropOn;
+	
+	if(SectionViewModel != nullptr)
+	{
+		SectionViewModel->GetOnRequestRename().BindSP(this, &SHierarchySection::TryEnterEditingMode);
+
+		SDropTarget::FArguments LeftDropTargetArgs;
+		SDropTarget::FArguments OntoDropTargetArgs;
+		SDropTarget::FArguments RightDropTargetArgs;
+
+		LeftDropTargetArgs
+			.OnAllowDrop(this, &SHierarchySection::OnCanAcceptDrop, EItemDropZone::AboveItem)
+			.OnDropped(this, &SHierarchySection::OnDroppedOn, EItemDropZone::AboveItem)
+			.VerticalImage(FAppStyle::GetNoBrush())
+			.HorizontalImage(FHierarchyEditorStyle::Get().GetBrush("HierarchyEditor.Stack.DropTarget.BorderHorizontal"));
+
+		OntoDropTargetArgs
+			.OnAllowDrop(this, &SHierarchySection::OnCanAcceptDrop, EItemDropZone::OntoItem)
+			.OnDropped(this, &SHierarchySection::OnDroppedOn, EItemDropZone::OntoItem);
+		
+		RightDropTargetArgs
+			.OnAllowDrop(this, &SHierarchySection::OnCanAcceptDrop, EItemDropZone::BelowItem)
+			.OnDropped(this, &SHierarchySection::OnDroppedOn, EItemDropZone::BelowItem)
+			.VerticalImage(FAppStyle::GetNoBrush())
+			.HorizontalImage(FHierarchyEditorStyle::Get().GetBrush("HierarchyEditor.Stack.DropTarget.BorderHorizontal"));
+
+		SetToolTipText(TAttribute<FText>::CreateSP(this, &SHierarchySection::GetTooltipText));
+
+		ChildSlot
+		[
+			SNew(SHorizontalBox)
+			+ SHorizontalBox::Slot()
+			.AutoWidth()
+			[
+				SNew(SSectionDragDropTarget, SectionViewModel, EItemDropZone::AboveItem)
+				.DropTargetArgs(LeftDropTargetArgs)
+			]
+			+ SHorizontalBox::Slot()
+			[
+				SNew(SSectionDragDropTarget, SectionViewModel, EItemDropZone::OntoItem)
+				.DropTargetArgs(OntoDropTargetArgs
+				.Content()
+				[
+					SAssignNew(CheckBox, SCheckBox)
+					.Visibility(EVisibility::HitTestInvisible)
+					.Style(FAppStyle::Get(), "DetailsView.SectionButton")
+					.OnCheckStateChanged(this, &SHierarchySection::OnSectionCheckChanged)
+					.IsChecked(this, &SHierarchySection::GetSectionCheckState)
+					.Padding(FMargin(8.f, 4.f))
+					[
+						SNew(SHorizontalBox)
+						+ SHorizontalBox::Slot()
+						.AutoWidth()
+						.Padding(2.f)
+						[
+							SNew(SImage).Image(SectionViewModel->GetSectionImage() != nullptr ? SectionViewModel->GetSectionImage() : FStyleDefaults::GetNoBrush())
+						]
+						+ SHorizontalBox::Slot()
+						[
+							SAssignNew(InlineEditableTextBlock, SInlineEditableTextBlock)
+							.Visibility(EVisibility::HitTestInvisible)
+							.Text(this, &SHierarchySection::GetText)
+							.OnTextCommitted(this, &SHierarchySection::OnRenameSection)
+							.OnVerifyTextChanged(this, &SHierarchySection::OnVerifySectionRename)
+							.IsSelected(this, &SHierarchySection::IsSectionSelected)
+							.IsReadOnly(this, &SHierarchySection::IsSectionReadOnly)
+						]
+					]
+				])
+			]
+			+ SHorizontalBox::Slot()
+			.AutoWidth()
+			[
+				SNew(SSectionDragDropTarget, SectionViewModel, EItemDropZone::BelowItem)
+				.DropTargetArgs(RightDropTargetArgs)				
+			]
+		];		
+	}
+	// if this section doesn't represent data, it's the "All" widget
+	else
+	{
+		ChildSlot
+		[
+			SNew(SDropTarget)
+			.OnAllowDrop(this, &SHierarchySection::OnCanAcceptDrop, EItemDropZone::OntoItem)
+			.OnDropped(this, &SHierarchySection::OnDroppedOn, EItemDropZone::OntoItem)
+			[
+				SNew(SCheckBox)
+				.Style(FAppStyle::Get(), "DetailsView.SectionButton")
+				.OnCheckStateChanged(this, &SHierarchySection::OnSectionCheckChanged)
+				.IsChecked(this, &SHierarchySection::GetSectionCheckState)
+				.Padding(FMargin(8.f, 4.f))
+				[
+					SNew(STextBlock)
+					.Text(LOCTEXT("AllSection", "All"))
+				]
+			]			
+		];
+	}
+}
+
+SHierarchySection::~SHierarchySection()
+{
+	SectionViewModel.Reset();
+}
+
+void SHierarchySection::TryEnterEditingMode() const
+{
+	if(SectionViewModel.IsValid() && SectionViewModel->CanRename())
+	{
+		InlineEditableTextBlock->EnterEditingMode();
+	}
+}
+
+TSharedPtr<FHierarchySectionViewModel> SHierarchySection::GetSectionViewModel()
+{
+	return SectionViewModel;
+}
+
+bool SHierarchySection::OnCanAcceptDrop(TSharedPtr<FDragDropOperation> DragDropOperation, EItemDropZone ItemDropZone) const
+{
+	if(bForbidDropOn)
+	{
+		return false;
+	}
+
+	if(DragDropOperation->IsOfType<FHierarchyDragDropOp>())
+	{
+		TSharedPtr<FHierarchyDragDropOp> HierarchyDragDropOp = StaticCastSharedPtr<FHierarchyDragDropOp>(DragDropOperation);
+		TSharedPtr<FHierarchyElementViewModel> DraggedItem = HierarchyDragDropOp->GetDraggedElement().Pin();
+		
+		if(SectionViewModel.IsValid())
+		{
+			return SectionViewModel->CanDropOn(DraggedItem, ItemDropZone).bCanPerform;
+		}
+		
+		// for the All section which has no valid view model, we simply do a check if the sections of the dragged categories are different
+		if(DragDropOperation->IsOfType<FHierarchyDragDropOp>())
+		{
+			if(const UHierarchyCategory* Category = Cast<UHierarchyCategory>(DraggedItem->GetData()))
+			{
+				return TryGetSectionData() != Category->GetSection();
+			}
+		}
+	}
+
+	return false;
+}
+
+FReply SHierarchySection::OnDroppedOn(const FGeometry&, const FDragDropEvent& DragDropEvent, EItemDropZone DropZone) const
+{
+	bDraggedOn = false;
+	if(TSharedPtr<FHierarchyDragDropOp> HierarchyDragDropOp = DragDropEvent.GetOperationAs<FHierarchyDragDropOp>())
+	{
+		TSharedPtr<FHierarchyElementViewModel> DraggedItem = HierarchyDragDropOp->GetDraggedElement().Pin();
+		if(SectionViewModel.IsValid())
+		{
+			SectionViewModel->OnDroppedOn(DraggedItem, DropZone);
+			return FReply::Handled();
+		}
+		
+		if(UHierarchyCategory* Category = Cast<UHierarchyCategory>(DraggedItem->GetDataMutable()))
+		{
+			Category->SetSection(nullptr);
+			TArray<UHierarchyCategory*> ChildrenCategories;
+			Category->GetChildrenOfType<>(ChildrenCategories, true);
+
+			for(UHierarchyCategory* ChildCategory : ChildrenCategories)
+			{
+				ChildCategory->SetSection(nullptr);
+			}
+
+			// we only need to reparent if the parent isn't already the root. This stops unnecessary reordering
+			if(HierarchyDragDropOp->GetDraggedElement().Pin()->GetParent() != HierarchyViewModel->GetHierarchyRootViewModel())
+			{
+				HierarchyViewModel->GetHierarchyRootViewModel()->ReparentToThis(HierarchyDragDropOp->GetDraggedElement().Pin());
+			}
+		
+			HierarchyViewModel->RefreshHierarchyView();
+
+			return FReply::Handled();
+		}
+	}
+
+	return FReply::Unhandled();
+}
+
+FReply SHierarchySection::OnMouseButtonDown(const FGeometry& MyGeometry, const FPointerEvent& MouseEvent)
+{
+	if(SectionViewModel.IsValid())
+	{
+		// we handle the event here so we can react on mouse button up
+		if(MouseEvent.IsMouseButtonDown(EKeys::RightMouseButton))
+		{
+			return FReply::Handled();
+		}
+		else if(MouseEvent.IsMouseButtonDown(EKeys::LeftMouseButton))
+		{			
+			OnSectionActivatedDelegate.ExecuteIfBound(SectionViewModel);
+			return FReply::Handled().DetectDrag(AsShared(), EKeys::LeftMouseButton).SetUserFocus(AsShared());	
+		}
+	}
+
+	return FReply::Unhandled();
+}
+
+FReply SHierarchySection::OnMouseButtonUp(const FGeometry& MyGeometry, const FPointerEvent& MouseEvent)
+{
+	if(SectionViewModel.IsValid() && SectionViewModel->IsForHierarchy())
+	{
+		if(MouseEvent.GetEffectingButton() == EKeys::RightMouseButton)
+		{			
+			// Show menu to choose getter vs setter
+			FSlateApplication::Get().PushMenu(
+				AsShared(),
+				FWidgetPath(),
+				SummonContextMenu({GetSectionViewModel()}),
+				FSlateApplication::Get().GetCursorPos(),
+				FPopupTransitionEffect(FPopupTransitionEffect::ContextMenu)
+			);
+				
+			OnSectionActivatedDelegate.ExecuteIfBound(SectionViewModel);			
+			return FReply::Handled();
+		}
+	}
+
+	return FReply::Unhandled();
+}
+
+void SHierarchySection::OnDragEnter(const FGeometry& MyGeometry, const FDragDropEvent& DragDropEvent)
+{
+	bDraggedOn = true;
+	if(DragDropEvent.GetOperationAs<FHierarchyDragDropOp>() && DragDropEvent.GetOperationAs<FSectionDragDropOp>() == nullptr)
+	{
+		RegisterActiveTimer(1.0f, FWidgetActiveTimerDelegate::CreateSP(this, &SHierarchySection::ActivateSectionIfDragging));
+	}
+}
+
+FReply SHierarchySection::OnDragDetected(const FGeometry& MyGeometry, const FPointerEvent& MouseEvent)
+{
+	if(MouseEvent.IsMouseButtonDown(EKeys::LeftMouseButton) && SectionViewModel.IsValid() && SectionViewModel->CanDrag().bCanPerform)
+	{
+		TSharedRef<FSectionDragDropOp> SectionDragDropOp = MakeShared<FSectionDragDropOp>(SectionViewModel);
+		SectionDragDropOp->Construct();
+		return FReply::Handled().BeginDragDrop(SectionDragDropOp);
+	}
+	
+	return FReply::Unhandled();
+}
+
+void SHierarchySection::OnDragLeave(const FDragDropEvent& DragDropEvent)
+{
+	bDraggedOn = false;
+	if(TSharedPtr<FHierarchyDragDropOp> HierarchyDragDropOp = DragDropEvent.GetOperationAs<FHierarchyDragDropOp>())
+	{
+		HierarchyDragDropOp->SetDescription(FText::GetEmpty());
+	}
+}
+
+FReply SHierarchySection::OnKeyDown(const FGeometry& MyGeometry, const FKeyEvent& InKeyEvent)
+{
+	if(SectionViewModel.IsValid() && InKeyEvent.GetKey() == EKeys::Delete && SectionViewModel->CanDelete())
+	{
+		SectionViewModel->Delete();
+		return FReply::Handled();
+	}
+
+	return FReply::Unhandled();
+}
+
+void SHierarchySection::OnMouseLeave(const FPointerEvent& MouseEvent)
+{
+	//bDraggedOn = false;
+	SCompoundWidget::OnMouseLeave(MouseEvent);
+}
+
+UHierarchySection* SHierarchySection::TryGetSectionData() const
+{
+	return SectionViewModel.IsValid() ? Cast<UHierarchySection>(SectionViewModel->GetDataMutable()) : nullptr;
+}
+
+FText SHierarchySection::GetText() const
+{
+	return SectionViewModel->GetSectionNameAsText();
+}
+
+FText SHierarchySection::GetTooltipText() const
+{
+	return SectionViewModel->GetSectionTooltip();
+}
+
+void SHierarchySection::OnRenameSection(const FText& Text, ETextCommit::Type CommitType) const
+{
+	FScopedTransaction Transaction(LOCTEXT("Transaction_Rename_Section", "Renamed hierarchy section"));
+	HierarchyViewModel->GetHierarchyRoot()->Modify();
+	
+	SectionViewModel->Rename(FName(Text.ToString()));
+}
+
+bool SHierarchySection::OnVerifySectionRename(const FText& NewName, FText& OutTooltip) const
+{
+	// this function shouldn't be used in case the section isn't valid but we'll make sure regardless
+	if(!SectionViewModel.IsValid())
+	{
+		return false;
+	}
+
+	if(SectionViewModel->GetSectionName().ToString() != NewName.ToString())
+	{
+		TArray<FString> SectionNames;
+
+		SectionNames.Add("All");
+		for(auto& Section : HierarchyViewModel->GetHierarchyRootViewModel()->GetSectionViewModels())
+		{
+			SectionNames.Add(Section->GetSectionName().ToString());
+		}
+
+		if(SectionNames.Contains(NewName.ToString()))
+		{
+			OutTooltip = LOCTEXT("HierarchySectionCantRename_Duplicate", "A section with that name already exists!");
+			return false;
+		}
+	}
+
+	return true;
+}
+
+bool SHierarchySection::IsSectionSelected() const
+{
+	return GetSectionCheckState() == ECheckBoxState::Checked ? true : false;
+}
+
+bool SHierarchySection::IsSectionReadOnly() const
+{
+	return SectionViewModel.IsValid() ? !SectionViewModel->CanRename() : true;
+}
+
+ECheckBoxState SHierarchySection::GetSectionCheckState() const
+{
+	return IsSectionActive.Get();
+}
+
+void SHierarchySection::OnSectionCheckChanged(ECheckBoxState NewState)
+{
+	OnSectionActivatedDelegate.ExecuteIfBound(SectionViewModel);
+}
+
+EActiveTimerReturnType SHierarchySection::ActivateSectionIfDragging(double CurrentTime, float DeltaTime) const
+{
+	if(bDraggedOn && FSlateApplication::Get().GetDragDroppingContent().IsValid() && FSlateApplication::Get().GetDragDroppingContent()->IsOfType<FHierarchyDragDropOp>())
+	{
+		if(IsSectionSelected() == false)
+		{
+			OnSectionActivatedDelegate.ExecuteIfBound(SectionViewModel);
+		}
+	}
+
+	return EActiveTimerReturnType::Stop;
+}
+
+void SHierarchyEditor::Construct(const FArguments& InArgs, TObjectPtr<UHierarchyViewModelBase> InHierarchyViewModel)
+{
+	HierarchyViewModel = InHierarchyViewModel;
+
+	// If the user hasn't called initialize themselves, we do it here, but ideally the user should do it themselves where appropriate
+	if(HierarchyViewModel->IsInitialized() == false)
+	{
+		HierarchyViewModel->Initialize();
+	}
+	
+	SourceRoot = NewObject<UHierarchyRoot>(GetTransientPackage(), MakeUniqueObjectName(GetTransientPackage(), UHierarchyRoot::StaticClass()));
+	TSharedPtr<FHierarchyElementViewModel> ViewModel = HierarchyViewModel->CreateViewModelForElement(SourceRoot, nullptr);
+	SourceRootViewModel = StaticCastSharedPtr<FHierarchyRootViewModel>(ViewModel);	
+	SourceRootViewModel->Initialize();
+	SourceRootViewModel->AddChildFilter(FHierarchyElementViewModel::FOnFilterChild::CreateSP(this, &SHierarchyEditor::FilterForSourceSection));
+	SourceRootViewModel->OnSyncPropagated().BindSP(this, &SHierarchyEditor::RequestRefreshSourceViewNextFrame, false);
+	SourceRootViewModel->OnSectionsChanged().BindSP(this, &SHierarchyEditor::RefreshSectionsView);
+
+	HierarchyViewModel->OnInitialized().BindSP(this, &SHierarchyEditor::Reinitialize);
+	HierarchyViewModel->OnNavigateToElementIdentityInHierarchyRequested().BindSP(this, &SHierarchyEditor::NavigateToHierarchyElement);
+	HierarchyViewModel->OnNavigateToElementInHierarchyRequested().BindSP(this, &SHierarchyEditor::NavigateToHierarchyElement);
+	HierarchyViewModel->OnRefreshSourceItemsRequested().BindSP(this, &SHierarchyEditor::RefreshSourceItems);
+	HierarchyViewModel->OnRefreshViewRequested().BindSP(this, &SHierarchyEditor::RefreshAllViews);
+	HierarchyViewModel->OnRefreshSourceView().BindSP(this, &SHierarchyEditor::RefreshSourceView);
+	HierarchyViewModel->OnRefreshHierarchyView().BindSP(this, &SHierarchyEditor::RefreshHierarchyView);
+	HierarchyViewModel->OnRefreshSectionsView().BindSP(this, &SHierarchyEditor::RefreshSectionsView);
+	HierarchyViewModel->OnHierarchySectionActivated().BindSP(this, &SHierarchyEditor::OnHierarchySectionActivated);
+	HierarchyViewModel->OnElementAdded().BindSP(this, &SHierarchyEditor::OnElementAdded);
+
+	BindToHierarchyRootViewModel();
+	
+	OnGenerateRowContentWidget = InArgs._OnGenerateRowContentWidget;
+	OnGenerateCustomDetailsPanelNameWidget = InArgs._OnGenerateCustomDetailsPanelNameWidget;
+	CategoryRowStyle = InArgs._CategoryRowStyle;
+	ItemRowStyle = InArgs._ItemRowStyle;
+	
+	if(!ensureMsgf(OnGenerateRowContentWidget.IsBound(), TEXT("Please add a function binding to the OnGenerateRowContentWidget slate event. Using default row content.")))
+	{
+		OnGenerateRowContentWidget = FOnGenerateRowContentWidget::CreateLambda([](TSharedRef<FHierarchyElementViewModel> HierarchyElement) -> TSharedRef<SWidget>
+		{
+			if(HierarchyElement->GetDataMutable()->IsA<UHierarchyCategory>())
+			{
+				TSharedRef<FHierarchyCategoryViewModel> TreeViewCategory = StaticCastSharedRef<FHierarchyCategoryViewModel>(HierarchyElement);
+				return SNew(SHierarchyCategory, TreeViewCategory);
+			}
+			
+			return SNew(STextBlock).Text(HierarchyElement->ToStringAsText());
+		});
+	}
+	SSplitter::FSlot* DetailsPanelSlot = nullptr;
+
+	TSharedRef<SWidget> AddSectionButton = SNew(SHorizontalBox)
+		+ SHorizontalBox::Slot()
+		.AutoWidth()
+		[
+			SNew(SButton)
+			.OnClicked(this, &SHierarchyEditor::OnAddSectionClicked)
+			.ButtonStyle(FHierarchyEditorStyle::Get(), "HierarchyEditor.ButtonStyle")
+			[
+				SNew(SHorizontalBox)
+				+ SHorizontalBox::Slot()
+				.AutoWidth()
+				.Padding(2.f)
+				[
+					SNew(SImage)
+					.Image(FAppStyle::Get().GetBrush("Icons.PlusCircle"))
+				]
+				+ SHorizontalBox::Slot()
+				.VAlign(VAlign_Center)
+				.AutoWidth()
+				[			
+					SNew(STextBlock)
+					.Text(LOCTEXT("AddSectionLabel","Add Section"))				
+				]
+			]
+		];
+
+	TSharedRef<SWidget> AddCategoryButton = SNew(SHorizontalBox)
+		+ SHorizontalBox::Slot()
+		.AutoWidth()
+		[
+			SNew(SButton)
+			.OnClicked(this, &SHierarchyEditor::OnAddCategoryClicked)
+			.ButtonStyle(FHierarchyEditorStyle::Get(), "HierarchyEditor.ButtonStyle")
+			[
+				SNew(SHorizontalBox)
+				+ SHorizontalBox::Slot()
+				.AutoWidth()
+				.Padding(2.f)
+				[
+					SNew(SImage)
+					.Image(FAppStyle::Get().GetBrush("Icons.PlusCircle"))
+				]
+				+ SHorizontalBox::Slot()
+				.VAlign(VAlign_Center)
+				.AutoWidth()
+				[			
+					SNew(STextBlock)
+					.Text(LOCTEXT("AddCategoryLabel","Add Category"))				
+				]
+			]
+		];
+	
+	ChildSlot
+	[
+		SNew(SBorder)
+		.Padding(0.f)
+		.BorderImage(FAppStyle::Get().GetBrush("Brushes.Panel"))
+		[
+			SNew(SSplitter)
+			.Orientation(Orient_Horizontal)
+			.PhysicalSplitterHandleSize(2.f)
+			+ SSplitter::Slot()
+			.Value(0.3f)
+			.MinSize(0.1f)
+			[
+				SNew(SVerticalBox)
+				+ SVerticalBox::Slot()
+				.AutoHeight()
+				.Padding(2.f)
+				[
+					SAssignNew(SourceSearchBox, SSearchBox)
+					.OnTextChanged(this, &SHierarchyEditor::OnSourceSearchTextChanged)
+					.OnTextCommitted(this, &SHierarchyEditor::OnSourceSearchTextCommitted)
+					.OnSearch(SSearchBox::FOnSearch::CreateSP(this, &SHierarchyEditor::OnSearchButtonClicked))
+					.DelayChangeNotificationsWhileTyping(true)
+					.SearchResultData(this, &SHierarchyEditor::GetSearchResultData)
+				]
+				+ SVerticalBox::Slot()
+				.AutoHeight()
+				.Padding(2.f)
+				[
+					SAssignNew(SourceSectionBox, SWrapBox)
+					.UseAllottedSize(true)
+				]
+				+ SVerticalBox::Slot()
+				.Padding(1.f, 2.f)
+				[
+					SAssignNew(SourceTreeView, STreeView<TSharedPtr<FHierarchyElementViewModel>>)
+					.TreeItemsSource(&GetSourceItems())
+					.OnSelectionChanged(this, &SHierarchyEditor::OnSelectionChanged, false)
+					.OnGenerateRow(this, &SHierarchyEditor::GenerateSourceItemRow)
+					.OnGetChildren_UObject(InHierarchyViewModel.Get(), &UHierarchyViewModelBase::OnGetChildren)
+					.OnItemToString_Debug_UObject(InHierarchyViewModel.Get(), &UHierarchyViewModelBase::OnElementToStringDebug)
+					.OnContextMenuOpening(this, &SHierarchyEditor::SummonContextMenuForSelectedRows, false)
+				]					
+			]
+			+ SSplitter::Slot()
+			.Value(0.4f)
+			.MinSize(0.1f)
+			[
+				SNew(SVerticalBox)
+				+ SVerticalBox::Slot()
+				.AutoHeight()
+				[
+					SNew(SVerticalBox)
+					+ SVerticalBox::Slot()
+					.AutoHeight()
+					[
+						SNew(SBorder)
+						.Padding(0.f)
+						.BorderImage(FAppStyle::Get().GetBrush("Brushes.Header"))
+						[
+							AddSectionButton
+						]
+					]
+					+ SVerticalBox::Slot()
+					.AutoHeight()
+					.Padding(1.f)
+					[
+						SAssignNew(HierarchySectionBox, SWrapBox)
+						.UseAllottedSize(true)
+					]				
+				]
+				+ SVerticalBox::Slot()
+				[
+					SNew(SVerticalBox)
+					+ SVerticalBox::Slot()
+					.AutoHeight()
+					[
+						SNew(SBorder)
+						.Padding(0.f)
+						.BorderImage(FAppStyle::Get().GetBrush("Brushes.Header"))
+						[
+							AddCategoryButton
+						]
+					]
+					+ SVerticalBox::Slot()
+					.FillHeight(0.1f)
+					.Padding(1.f, 4.f, 1.f, 0.f)
+					[
+						SNew(SDropTarget)
+						.OnDropped(this, &SHierarchyEditor::HandleHierarchyRootDrop)
+						.OnAllowDrop(this, &SHierarchyEditor::OnCanDropOnRoot)
+						.OnDragEnter(this, &SHierarchyEditor::OnRootDragEnter)
+						.OnDragLeave(this, &SHierarchyEditor::OnRootDragLeave)
+						[
+							SNew(SBorder)
+							.Padding(0.f)
+							.BorderImage(FAppStyle::Get().GetBrush("Brushes.Recessed"))
+							[
+								SNew(SBorder)
+								.Padding(1.f)
+								.BorderImage(FAppStyle::GetBrush("DashedBorder"))
+								.BorderBackgroundColor(FLinearColor(0.2f, 0.2f, 0.2f, 0.5f))
+								[
+									SNew(SBox)
+									.HAlign(HAlign_Center)
+									.VAlign(VAlign_Center)
+									[
+										SNew(SImage)
+										.Image(FHierarchyEditorStyle::Get().GetBrush("HierarchyEditor.RootDropIcon"))
+										.ColorAndOpacity(this, &SHierarchyEditor::GetRootIconColor)
+									]
+								]
+							]
+						]
+					]
+					+ SVerticalBox::Slot()
+					.Padding(1.f, 0.f)
+					[
+						SAssignNew(HierarchyTreeView, STreeView<TSharedPtr<FHierarchyElementViewModel>>)
+						.TreeItemsSource(&InHierarchyViewModel->GetHierarchyItems())
+						.OnSelectionChanged(this, &SHierarchyEditor::OnSelectionChanged, true)
+						.OnGenerateRow(this, &SHierarchyEditor::GenerateHierarchyItemRow)
+						.OnGetChildren_UObject(InHierarchyViewModel.Get(), &UHierarchyViewModelBase::OnGetChildren)
+						.OnItemToString_Debug_UObject(InHierarchyViewModel.Get(), &UHierarchyViewModelBase::OnElementToStringDebug)
+						.OnContextMenuOpening(this, &SHierarchyEditor::SummonContextMenuForSelectedRows, true)
+					]
+				]
+			]
+			+ SSplitter::Slot()
+			.Value(0.3f)
+			.MinSize(0.1f)
+			.Expose(DetailsPanelSlot)		
+		]
+	];
+
+	if(InHierarchyViewModel->SupportsDetailsPanel())
+	{
+		FPropertyEditorModule& PropertyEditorModule = FModuleManager::Get().LoadModuleChecked<FPropertyEditorModule>("PropertyEditor");
+
+		FDetailsViewArgs DetailsViewArgs;
+		DetailsViewArgs.NameAreaSettings = FDetailsViewArgs::ENameAreaSettings::ObjectsUseNameArea;
+		DetailsViewArgs.bShowObjectLabel = false;
+		DetailsViewArgs.bAllowSearch = false;
+		DetailsViewArgs.NotifyHook = this;
+
+		DetailsPanel = PropertyEditorModule.CreateDetailView(DetailsViewArgs);
+		
+		if(OnGenerateCustomDetailsPanelNameWidget.IsBound())
+		{
+			TSharedRef<SWidget> CustomDetailsPanelNameWidget = OnGenerateCustomDetailsPanelNameWidget.Execute(nullptr);
+			DetailsPanel->SetNameAreaCustomContent(CustomDetailsPanelNameWidget);
+		}
+
+		DetailsPanel->SetIsPropertyEditingEnabledDelegate(FIsPropertyEditingEnabled::CreateSP(this, &SHierarchyEditor::IsDetailsPanelEditingAllowed));
+		
+		for(auto& Customizations : HierarchyViewModel->GetInstanceCustomizations())
+		{
+			DetailsPanel->RegisterInstancedCustomPropertyLayout(Customizations.Key, Customizations.Value);
+		}
+		
+		DetailsPanelSlot->AttachWidget(DetailsPanel.ToSharedRef());
+	}
+		
+	HierarchyViewModel->GetCommands()->MapAction(FGenericCommands::Get().Rename,
+		FExecuteAction::CreateSP(this, &SHierarchyEditor::RequestRenameSelectedItem),
+		FCanExecuteAction::CreateSP(this, &SHierarchyEditor::CanRequestRenameSelectedItem),
+		FIsActionChecked(), FIsActionButtonVisible::CreateSP(this, &SHierarchyEditor::CanRequestRenameSelectedItem));
+	
+	HierarchyViewModel->GetCommands()->MapAction(FGenericCommands::Get().Delete,
+		FExecuteAction::CreateSP(this, &SHierarchyEditor::DeleteSelectedHierarchyItems),
+		FCanExecuteAction::CreateSP(this, &SHierarchyEditor::CanDeleteSelectedElements),
+		FIsActionChecked(), FIsActionButtonVisible::CreateSP(this, &SHierarchyEditor::CanDeleteSelectedElements));
+	
+	HierarchyViewModel->GetCommands()->MapAction(FHierarchyEditorCommands::Get().FindInHierarchy,
+		FExecuteAction::CreateSP(this, &SHierarchyEditor::NavigateToMatchingHierarchyElementFromSelectedSourceElement),
+		FCanExecuteAction::CreateSP(this, &SHierarchyEditor::CanNavigateToMatchingHierarchyElementFromSelectedSourceElement));
+	
+	HierarchyViewModel->ForceFullRefresh();
+}
+
+SHierarchyEditor::~SHierarchyEditor()
+{
+	SourceSearchResults.Empty();
+	FocusedSearchResult.Reset();
+	
+	ClearSourceItems();
+	
+	if(HierarchyViewModel.IsValid())
+	{
+		HierarchyViewModel->OnInitialized().Unbind();
+		HierarchyViewModel->OnNavigateToElementIdentityInHierarchyRequested().Unbind();
+		HierarchyViewModel->OnNavigateToElementInHierarchyRequested().Unbind();
+		HierarchyViewModel->OnRefreshSourceItemsRequested().Unbind();
+		HierarchyViewModel->OnRefreshViewRequested().Unbind();
+		HierarchyViewModel->OnRefreshSourceView().Unbind();
+		HierarchyViewModel->OnRefreshHierarchyView().Unbind();
+		HierarchyViewModel->OnRefreshSectionsView().Unbind();
+		HierarchyViewModel->OnHierarchySectionActivated().Unbind();
+		HierarchyViewModel->OnElementAdded().Unbind();
+
+		UnbindFromHierarchyRootViewModel();
+
+		HierarchyViewModel->GetCommands()->UnmapAction(FGenericCommands::Get().Delete);
+		HierarchyViewModel->GetCommands()->UnmapAction(FGenericCommands::Get().Rename);
+		HierarchyViewModel->GetCommands()->UnmapAction(FHierarchyEditorCommands::Get().FindInHierarchy);
+	}
+	
+	SourceRootViewModel->OnSyncPropagated().Unbind();
+	SourceRootViewModel->OnSectionsChanged().Unbind();
+	SourceRootViewModel.Reset();
+	SourceRoot->ConditionalBeginDestroy();
+	SourceRoot = nullptr;
+}
+
+void SHierarchyEditor::RefreshSourceItems()
+{
+	SourceRoot->EmptyAllData();
+	HierarchyViewModel->PrepareSourceItems(SourceRoot, SourceRootViewModel);
+	SourceRootViewModel->SyncViewModelsToData();
+	RefreshSourceView(false);
+	RefreshSectionsView();
+}
+
+void SHierarchyEditor::RefreshAllViews(bool bFullRefresh)
+{
+	RefreshSourceView(bFullRefresh);
+	RefreshHierarchyView(bFullRefresh);
+	RefreshSectionsView();
+}
+
+void SHierarchyEditor::RequestRefreshAllViewsNextFrame(bool bFullRefresh)
+{
+	RequestRefreshSourceViewNextFrame(bFullRefresh);
+	RequestRefreshHierarchyViewNextFrame(bFullRefresh);
+	RequestRefreshSectionsViewNextFrame();
+}
+
+FReply SHierarchyEditor::OnKeyDown(const FGeometry& MyGeometry, const FKeyEvent& InKeyEvent)
+{
+	UHierarchyMenuContext* HierarchyMenuContext = NewObject<UHierarchyMenuContext>();
+	 
+	FToolMenuContext Context;
+	Context.AddObject(HierarchyMenuContext);
+	
+	UToolMenus::Get()->GenerateMenu("HierarchyMenu", Context);
+	return HierarchyViewModel->GetCommands()->ProcessCommandBindings(InKeyEvent) ? FReply::Handled() : FReply::Unhandled();
+}
+
+FReply SHierarchyEditor::OnMouseButtonDown(const FGeometry& MyGeometry, const FPointerEvent& MouseEvent)
+{
+	// we catch any mouse button down event so that we can continue using our commands
+	return FReply::Handled().SetUserFocus(AsShared(), EFocusCause::Mouse, true);
+}
+
+FReply SHierarchyEditor::OnMouseButtonUp(const FGeometry& MyGeometry, const FPointerEvent& MouseEvent)
+{
+	return FReply::Handled().SetUserFocus(AsShared(), EFocusCause::Mouse, true);
+}
+
+FReply SHierarchyEditor::OnAddCategoryClicked() const
+{
+	TArray<TSharedPtr<FHierarchyElementViewModel>> SelectedItems;
+	HierarchyTreeView->GetSelectedItems(SelectedItems);
+
+	// We can only add categories under categories or the root
+	if(SelectedItems.Num() == 1 && SelectedItems[0]->GetData()->IsA<UHierarchyCategory>())
+	{
+		HierarchyViewModel.Get()->AddCategory(SelectedItems[0]);
+	}
+	else
+	{
+		HierarchyViewModel.Get()->AddCategory(nullptr);
+	}
+	
+	return FReply::Handled();
+}
+
+FReply SHierarchyEditor::OnAddSectionClicked() const
+{
+	HierarchyViewModel.Get()->AddSection();
+	return FReply::Handled();
+}
+
+TSharedPtr<SWidget> SHierarchyEditor::SummonContextMenuForSelectedRows(bool bFromHierarchy) const
+{
+	TArray<TSharedPtr<FHierarchyElementViewModel>> ViewModels;
+	if(bFromHierarchy)
+	{
+		HierarchyTreeView->GetSelectedItems(ViewModels);
+	}
+	else
+	{
+		SourceTreeView->GetSelectedItems(ViewModels);
+	}
+
+	if(ViewModels.IsEmpty())
+	{
+		return nullptr;
+	}
+	
+	return SummonContextMenu(ViewModels);
+}
+
+void SHierarchyEditor::RefreshSourceView(bool bFullRefresh) const
+{
+	SourceTreeView->SetTreeItemsSource(&GetSourceItems());
+	if(bFullRefresh)
+	{
+		SourceTreeView->RebuildList();
+	}
+	else
+	{
+		SourceTreeView->RequestTreeRefresh();
+	}
+}
+
+void SHierarchyEditor::RequestRefreshSourceViewNextFrame(bool bFullRefresh)
+{
+	if(!RefreshSourceViewNextFrameHandle.IsValid())
+	{
+		RefreshSourceViewNextFrameHandle = RegisterActiveTimer(0.f, FWidgetActiveTimerDelegate::CreateLambda([this, bFullRefresh](double CurrentTime, float DeltaTime)
+		{
+			RefreshSourceView(bFullRefresh);
+			RefreshSourceViewNextFrameHandle.Reset();
+			return EActiveTimerReturnType::Stop;
+		}));
+	}
+}
+
+void SHierarchyEditor::RefreshHierarchyView(bool bFullRefresh) const
+{
+	// the top layer objects might have changed due to filtering. We need to refresh these too.
+	HierarchyTreeView->SetTreeItemsSource(&HierarchyViewModel->GetHierarchyItems());
+	if(bFullRefresh)
+	{
+		HierarchyTreeView->RebuildList();
+	}
+	else
+	{
+		HierarchyTreeView->RequestTreeRefresh();
+	}
+}
+
+void SHierarchyEditor::RequestRefreshHierarchyViewNextFrame(bool bFullRefresh)
+{
+	if(!RefreshHierarchyViewNextFrameHandle.IsValid())
+	{
+		RefreshHierarchyViewNextFrameHandle = RegisterActiveTimer(0.f, FWidgetActiveTimerDelegate::CreateLambda([this, bFullRefresh](double CurrentTime, float DeltaTime)
+		{
+			RefreshHierarchyView(bFullRefresh);
+			RefreshHierarchyViewNextFrameHandle.Reset();
+			return EActiveTimerReturnType::Stop;
+		}));
+	}
+}
+
+void SHierarchyEditor::RefreshSectionsView()
+{
+	SourceSectionBox->ClearChildren();
+	HierarchySectionBox->ClearChildren();
+
+	for (TSharedPtr<FHierarchySectionViewModel>& SourceSection : SourceRootViewModel->GetSectionViewModels())
+	{
+		TSharedPtr<SHierarchySection> SectionWidget = SNew(SHierarchySection, SourceSection, HierarchyViewModel)
+		.IsSectionActive_Lambda([this, SourceSection]()
+		{
+			return GetActiveSourceSection() == SourceSection ? ECheckBoxState::Checked : ECheckBoxState::Unchecked;
+		})
+		.OnSectionActivated_Lambda([this](TSharedPtr<FHierarchySectionViewModel> SectionViewModel)
+		{
+			SetActiveSourceSection(SectionViewModel);
+		});
+		
+		SourceSectionBox->AddSlot()
+		.HAlign(HAlign_Center)
+		.VAlign(VAlign_Center)
+		[
+			SectionWidget.ToSharedRef()		
+		];
+	}
+
+	if(SourceRootViewModel->GetSectionViewModels().Num() > 0)
+	{
+		TSharedPtr<SHierarchySection> DefaultSourceSection = SNew(SHierarchySection, nullptr, HierarchyViewModel)
+		.IsSectionActive_Lambda([this]()
+		{
+			return GetActiveSourceSection() == nullptr ? ECheckBoxState::Checked : ECheckBoxState::Unchecked;
+		})
+		.OnSectionActivated(this, &SHierarchyEditor::SetActiveSourceSection)
+		// we forbid drop on here as the 'All' sections don't have a valid view model to determine it instead
+		.bForbidDropOn(true);
+
+		SourceSectionBox->AddSlot()
+		.HAlign(HAlign_Center)
+		.VAlign(VAlign_Center)
+		[
+			DefaultSourceSection.ToSharedRef()
+		];
+	}
+	
+	for (TSharedPtr<FHierarchySectionViewModel>& HierarchySection : HierarchyViewModel->GetHierarchyRootViewModel()->GetSectionViewModels())
+	{
+		TSharedPtr<SHierarchySection> SectionWidget = SNew(SHierarchySection, HierarchySection, HierarchyViewModel)
+		.IsSectionActive_Lambda([this, HierarchySection]()
+		{
+			return HierarchyViewModel->GetActiveHierarchySection() == HierarchySection ? ECheckBoxState::Checked : ECheckBoxState::Unchecked;
+		})
+		.OnSectionActivated_Lambda([this](TSharedPtr<FHierarchySectionViewModel> SectionViewModel)
+		{
+			HierarchyViewModel->SetActiveHierarchySection(SectionViewModel);
+		});
+		
+		HierarchySectionBox->AddSlot()
+		.HAlign(HAlign_Center)
+		.VAlign(VAlign_Center)
+		[
+			SectionWidget.ToSharedRef()		
+		];
+	}
+	
+	TSharedPtr<SHierarchySection> DefaultHierarchySection = SNew(SHierarchySection, nullptr, HierarchyViewModel)
+	.IsSectionActive_Lambda([this]()
+	{
+		return HierarchyViewModel->GetActiveHierarchySection() == nullptr ? ECheckBoxState::Checked : ECheckBoxState::Unchecked;
+	})
+	.OnSectionActivated_Lambda([this](TSharedPtr<FHierarchySectionViewModel> SectionViewModel)
+	{
+		HierarchyViewModel->SetActiveHierarchySection(SectionViewModel);
+	});
+
+	HierarchySectionBox->AddSlot()
+	.HAlign(HAlign_Center)
+	.VAlign(VAlign_Center)
+	[
+		DefaultHierarchySection.ToSharedRef()
+	];
+}
+
+void SHierarchyEditor::RequestRefreshSectionsViewNextFrame()
+{
+	if(!RefreshSectionsViewNextFrameHandle.IsValid())
+	{
+		RefreshSectionsViewNextFrameHandle = RegisterActiveTimer(0.f, FWidgetActiveTimerDelegate::CreateLambda([this](double CurrentTime, float DeltaTime)
+		{
+			RefreshSectionsView();
+			RefreshSectionsViewNextFrameHandle.Reset();
+			return EActiveTimerReturnType::Stop;
+		}));
+	}
+}
+
+void SHierarchyEditor::NavigateToHierarchyElement(FHierarchyElementIdentity Identity) const
+{
+	if(TSharedPtr<FHierarchyElementViewModel> ViewModel = HierarchyViewModel->GetHierarchyRootViewModel()->FindViewModelForChild(Identity, true))
+	{
+		NavigateToHierarchyElement(ViewModel);
+	}
+}
+
+void SHierarchyEditor::NavigateToHierarchyElement(TSharedPtr<FHierarchyElementViewModel> Item) const
+{
+	TArray<TSharedPtr<FHierarchyElementViewModel>> ParentChain;
+	for(TWeakPtr<FHierarchyElementViewModel> Parent = Item->GetParent(); Parent.IsValid(); Parent = Parent.Pin()->GetParent())
+	{
+		ParentChain.Add(Parent.Pin());
+	}
+
+	for(int32 ParentIndex = ParentChain.Num()-1; ParentIndex >= 0; ParentIndex--)
+	{
+		HierarchyTreeView->SetItemExpansion(ParentChain[ParentIndex], true);
+	}
+		
+	HierarchyTreeView->SetSelection(Item);
+	HierarchyTreeView->RequestScrollIntoView(Item);
+}
+
+bool SHierarchyEditor::IsItemSelected(TSharedPtr<FHierarchyElementViewModel> Item) const
+{
+	return HierarchyTreeView->IsItemSelected(Item);
+}
+
+TSharedRef<ITableRow> SHierarchyEditor::GenerateSourceItemRow(TSharedPtr<FHierarchyElementViewModel> HierarchyItem, const TSharedRef<STableViewBase>& TableViewBase)
+{
+	return SNew(STableRow<TSharedPtr<FHierarchyElementViewModel>>, TableViewBase)
+	.Style(HierarchyItem->GetData()->IsA<UHierarchyItem>() ? ItemRowStyle : CategoryRowStyle)
+	.OnDragDetected(HierarchyItem.ToSharedRef(), &FHierarchyElementViewModel::OnDragDetected, true)
+	.Padding(FMargin(2))
+	[
+		SNew(SHorizontalBox)
+		+ SHorizontalBox::Slot()
+		.AutoWidth()
+		.Padding(1.f)
+		.HAlign(HAlign_Center)
+		.VAlign(VAlign_Center)
+		[
+			SNew(SBox)
+			.HeightOverride(10.f)
+			.WidthOverride(10.f)
+			.Visibility_Lambda([HierarchyItem, this]()
+			{
+				TArray<TSharedPtr<FHierarchyElementViewModel>> AllChildren;
+				HierarchyItem->GetChildrenViewModelsForType<UHierarchyElement, FHierarchyElementViewModel>(AllChildren, true);
+
+				bool bCanDrag = HierarchyItem->GetHierarchyViewModel()->GetHierarchyRootViewModel()->FindViewModelForChild(HierarchyItem->GetData()->GetPersistentIdentity(), true) == nullptr;			
+
+				if(bCanDrag)
+				{
+					for(TSharedPtr<FHierarchyElementViewModel>& ItemViewModel : AllChildren)
+					{
+						if(HierarchyItem->GetHierarchyViewModel()->GetHierarchyRootViewModel()->FindViewModelForChild(ItemViewModel->GetData()->GetPersistentIdentity(), true) != nullptr)
+						{
+							bCanDrag = false;
+							break;
+						}
+					}
+				}
+
+				return bCanDrag ? EVisibility::Collapsed : EVisibility::Visible;
+			})
+			[
+				SNew(SImage)
+				.Image(FAppStyle::GetBrush("Icons.Lock"))
+				.ToolTipText(LOCTEXT("CantDragItemAlreadyInHierarchyTooltip", "This item already exists within the hierarchy and can not be dragged. Drag the existing one within the hierarchy directly."))
+			]
+		]
+		+ SHorizontalBox::Slot()
+		[
+			OnGenerateRowContentWidget.Execute(HierarchyItem.ToSharedRef())
+		]
+	];
+}
+
+TSharedRef<ITableRow> SHierarchyEditor::GenerateHierarchyItemRow(TSharedPtr<FHierarchyElementViewModel> HierarchyItem, const TSharedRef<STableViewBase>& TableViewBase)
+{
+	return SNew(STableRow<TSharedPtr<FHierarchyElementViewModel>>, TableViewBase)
+	.Style(HierarchyItem->GetData()->IsA<UHierarchyItem>() ? ItemRowStyle : CategoryRowStyle)
+	.OnAcceptDrop(HierarchyItem.ToSharedRef(), &FHierarchyElementViewModel::OnDroppedOnRow)
+	.OnCanAcceptDrop(HierarchyItem.ToSharedRef(), &FHierarchyElementViewModel::OnCanRowAcceptDrop)
+	.OnDragDetected(HierarchyItem.ToSharedRef(), &FHierarchyElementViewModel::OnDragDetected, false)
+	.OnDragLeave(HierarchyItem.ToSharedRef(), &FHierarchyElementViewModel::OnRowDragLeave)
+	.Padding(FMargin(2))
+	[
+		SNew(SHorizontalBox)
+		+ SHorizontalBox::Slot()
+		.AutoWidth()
+		.Padding(1.f)
+		.HAlign(HAlign_Center)
+		.VAlign(VAlign_Center)
+		[
+			SNew(SBox)
+			.HeightOverride(10.f)
+			.WidthOverride(10.f)
+			.Visibility_Lambda([HierarchyItem]()
+			{
+				return HierarchyItem->IsEditableByUser().bCanPerform ? EVisibility::Collapsed : EVisibility::Visible;
+			})
+			[
+				SNew(SImage)
+				.Image(FAppStyle::GetBrush("Icons.Lock"))
+				.ToolTipText_Lambda([HierarchyItem]()
+				{
+					FHierarchyElementViewModel::FCanPerformActionResults IsEditableResults = HierarchyItem->IsEditableByUser();
+					if(IsEditableResults.bCanPerform == false)
+					{
+						return IsEditableResults.CanPerformMessage;
+					}
+
+					return FText::GetEmpty();
+				})
+			]
+		]
+		+ SHorizontalBox::Slot()
+		[
+			OnGenerateRowContentWidget.Execute(HierarchyItem.ToSharedRef())
+		]
+	];
+}
+
+bool SHierarchyEditor::FilterForSourceSection(TSharedPtr<const FHierarchyElementViewModel> ItemViewModel) const
+{
+	if(ActiveSourceSection.IsValid())
+	{
+		return GetActiveSourceSectionData() == ItemViewModel->GetSection();
+	}
+
+	return true;
+}
+
+void SHierarchyEditor::Reinitialize()
+{
+	// the hierarchy root view model has been recreated if the view model reinitialized. Therefore we update the bindings.
+	BindToHierarchyRootViewModel();
+	RefreshSourceItems();
+	RefreshAllViews(true);
+}
+
+void SHierarchyEditor::BindToHierarchyRootViewModel()
+{
+	HierarchyViewModel->GetHierarchyRootViewModel()->OnSyncPropagated().BindSP(this, &SHierarchyEditor::RequestRefreshHierarchyViewNextFrame, false);
+	HierarchyViewModel->GetHierarchyRootViewModel()->OnSectionsChanged().BindSP(this, &SHierarchyEditor::RefreshSectionsView);
+	HierarchyViewModel->GetHierarchyRootViewModel()->OnSectionAdded().BindSP(this, &SHierarchyEditor::OnHierarchySectionAdded);
+	HierarchyViewModel->GetHierarchyRootViewModel()->OnSectionDeleted().BindSP(this, &SHierarchyEditor::OnHierarchySectionDeleted);
+}
+
+void SHierarchyEditor::UnbindFromHierarchyRootViewModel() const
+{
+	if(HierarchyViewModel->GetHierarchyRootViewModel().IsValid())
+	{
+		HierarchyViewModel->GetHierarchyRootViewModel()->OnSyncPropagated().Unbind();
+		HierarchyViewModel->GetHierarchyRootViewModel()->OnSectionsChanged().Unbind();
+		HierarchyViewModel->GetHierarchyRootViewModel()->OnSectionAdded().Unbind();
+		HierarchyViewModel->GetHierarchyRootViewModel()->OnSectionDeleted().Unbind();
+	}
+}
+
+const TArray<TSharedPtr<FHierarchyElementViewModel>>& SHierarchyEditor::GetSourceItems() const
+{
+	return SourceRootViewModel->GetFilteredChildren();
+}
+
+bool SHierarchyEditor::IsDetailsPanelEditingAllowed() const
+{
+	return SelectedDetailsPanelItemViewModel.IsValid() && SelectedDetailsPanelItemViewModel.Pin()->IsEditableByUser().bCanPerform;
+}
+
+void SHierarchyEditor::RequestRenameSelectedItem()
+{
+	TArray<TSharedPtr<FHierarchyElementViewModel>> SelectedItems = HierarchyTreeView->GetSelectedItems();
+
+	if(SelectedItems.Num() == 0)
+	{
+		TSharedPtr<FHierarchySectionViewModel> ActiveSection = HierarchyViewModel->GetActiveHierarchySection();
+		if(ActiveSection)
+		{
+			SelectedItems = { ActiveSection };
+		}
+	}
+	
+	if(SelectedItems.Num() == 1)
+	{
+		SelectedItems[0]->RequestRename();
+	}
+}
+
+bool SHierarchyEditor::CanRequestRenameSelectedItem() const
+{
+	TArray<TSharedPtr<FHierarchyElementViewModel>> SelectedItems = HierarchyTreeView->GetSelectedItems();
+
+	if(SelectedItems.Num() == 0)
+	{
+		TSharedPtr<FHierarchySectionViewModel> ActiveSection = HierarchyViewModel->GetActiveHierarchySection();
+		if(ActiveSection)
+		{
+			SelectedItems = { ActiveSection };
+		}
+	}
+	
+	if(SelectedItems.Num() == 1)
+	{
+		return SelectedItems[0]->CanRename();
+	}
+
+	return false;
+}
+
+void SHierarchyEditor::ClearSourceItems() const
+{
+	SourceRoot->GetChildrenMutable().Empty();
+	SourceRoot->GetSectionDataMutable().Empty();
+	SourceRootViewModel->GetChildrenMutable().Empty();
+	SourceRootViewModel->GetSectionViewModels().Empty();
+}
+
+void SHierarchyEditor::DeleteItems(TArray<TSharedPtr<FHierarchyElementViewModel>> ItemsToDelete) const
+{
+	TArray<FHierarchyElementIdentity> DeletionIdentities;
+	
+	for(TSharedPtr<FHierarchyElementViewModel>& SelectedItem : ItemsToDelete)
+	{
+		DeletionIdentities.Add(SelectedItem->GetData()->GetPersistentIdentity());
+	}
+
+	HierarchyViewModel->DeleteElementsWithIdentities(DeletionIdentities);
+}
+
+void SHierarchyEditor::DeleteSelectedHierarchyItems() const
+{	
+	TArray<TSharedPtr<FHierarchyElementViewModel>> SelectedElements = HierarchyTreeView->GetSelectedItems();
+	
+	if(SelectedElements.Num() == 0)
+	{
+		TSharedPtr<FHierarchySectionViewModel> ActiveSection = HierarchyViewModel->GetActiveHierarchySection();
+		if(ActiveSection)
+		{
+			SelectedElements = { ActiveSection };
+		}
+	}
+	
+	DeleteItems(SelectedElements);
+}
+
+bool SHierarchyEditor::CanDeleteSelectedElements() const
+{
+	TArray<TSharedPtr<FHierarchyElementViewModel>> SelectedElements = HierarchyTreeView->GetSelectedItems();
+
+	if(SelectedElements.Num() == 0)
+	{
+		TSharedPtr<FHierarchySectionViewModel> ActiveSection = HierarchyViewModel->GetActiveHierarchySection();
+		if(ActiveSection)
+		{
+			SelectedElements = { ActiveSection };
+		}
+	}
+	
+	if(SelectedElements.Num() > 0)
+	{
+		bool bCanDelete = true;
+		for(TSharedPtr<FHierarchyElementViewModel> SelectedElement : SelectedElements)
+		{
+			bCanDelete &= SelectedElements[0]->CanDelete();
+		}
+
+		return bCanDelete;
+	}
+
+	return false;
+}
+
+void SHierarchyEditor::NavigateToMatchingHierarchyElementFromSelectedSourceElement() const
+{
+	TArray<TSharedPtr<FHierarchyElementViewModel>> SelectedElements = SourceTreeView->GetSelectedItems();
+
+	if(SelectedElements.Num() != 1)
+	{
+		return;
+	}
+
+	TSharedPtr<FHierarchyElementViewModel> SelectedElement = SelectedElements[0];
+
+	if(TSharedPtr<FHierarchyElementViewModel> MatchingViewModelInHierarchy = HierarchyViewModel->GetHierarchyRootViewModel()->FindViewModelForChild(SelectedElement->GetData()->GetPersistentIdentity(), true))
+	{
+		NavigateToHierarchyElement(MatchingViewModelInHierarchy.ToSharedRef());
+	}
+}
+
+bool SHierarchyEditor::CanNavigateToMatchingHierarchyElementFromSelectedSourceElement() const
+{
+	TArray<TSharedPtr<FHierarchyElementViewModel>> SelectedElements = SourceTreeView->GetSelectedItems();
+
+	if(SelectedElements.Num() != 1)
+	{
+		return false;
+	}
+
+	TSharedPtr<FHierarchyElementViewModel> SelectedElement = SelectedElements[0];
+	if(SelectedElement->IsForHierarchy())
+	{
+		return false;
+	}
+	
+	if(TSharedPtr<FHierarchyElementViewModel> MatchingViewModelInHierarchy = HierarchyViewModel->GetHierarchyRootViewModel()->FindViewModelForChild(SelectedElement->GetData()->GetPersistentIdentity(), true))
+	{
+		return MatchingViewModelInHierarchy.IsValid();
+	}
+	
+	return false;
+}
+
+void SHierarchyEditor::DeleteActiveSection() const
+{
+	if(TSharedPtr<FHierarchySectionViewModel> ActiveHierarchySectionViewModel = HierarchyViewModel->GetActiveHierarchySection())
+	{
+		DeleteItems({ActiveHierarchySectionViewModel});
+	}
+}
+
+bool SHierarchyEditor::CanDeleteActiveSection() const
+{
+	return HierarchyViewModel->GetActiveHierarchySection() != nullptr;
+}
+
+void SHierarchyEditor::OnElementAdded(TSharedPtr<FHierarchyElementViewModel> AddedItem)
+{
+	// when a new item is created (opposed to dragged & dropped from source view, i.e. only categories so far)
+	// we make sure to request a tree refresh, select the row, and request a pending rename since the widget will created a frame later
+	if(AddedItem->GetData()->IsA<UHierarchyItem>() || AddedItem->GetData()->IsA<UHierarchyCategory>())
+	{
+		HierarchyTreeView->RequestTreeRefresh();
+		NavigateToHierarchyElement(AddedItem);
+	}
+	else if(AddedItem->GetData()->IsA<UHierarchySection>())
+	{
+		RefreshSectionsView();
+	}
+
+	AddedItem->RequestRenamePending();
+}
+
+void SHierarchyEditor::OnHierarchySectionActivated(TSharedPtr<FHierarchySectionViewModel> Section)
+{
+	OnSelectionChanged(Section, ESelectInfo::Direct, true);
+}
+
+void SHierarchyEditor::OnSourceSectionActivated(TSharedPtr<FHierarchySectionViewModel> Section)
+{
+	OnSelectionChanged(Section, ESelectInfo::Direct, false);
+	RunSourceSearch();
+}
+
+void SHierarchyEditor::OnHierarchySectionAdded(TSharedPtr<FHierarchySectionViewModel> AddedSection)
+{
+	HierarchyViewModel->SetActiveHierarchySection(AddedSection);
+	AddedSection->RequestRenamePending();
+}
+
+void SHierarchyEditor::OnHierarchySectionDeleted(TSharedPtr<FHierarchySectionViewModel> DeletedSection)
+{
+	if(HierarchyViewModel->GetActiveHierarchySection() == DeletedSection)
+	{
+		HierarchyViewModel->SetActiveHierarchySection(nullptr);
+	}
+}
+
+void SHierarchyEditor::SetActiveSourceSection(TSharedPtr<FHierarchySectionViewModel> Section)
+{
+	ActiveSourceSection = Section;	
+	RefreshSourceView(true);
+	OnSourceSectionActivated(Section);
+}
+
+TSharedPtr<FHierarchySectionViewModel> SHierarchyEditor::GetActiveSourceSection() const
+{
+	return ActiveSourceSection.IsValid() ? ActiveSourceSection.Pin() : nullptr;
+}
+
+UHierarchySection* SHierarchyEditor::GetActiveSourceSectionData() const
+{
+	return ActiveSourceSection.IsValid() ? ActiveSourceSection.Pin()->GetDataMutable<UHierarchySection>() : nullptr;
+}
+
+void SHierarchyEditor::OnSelectionChanged(TSharedPtr<FHierarchyElementViewModel> HierarchyItem, ESelectInfo::Type Type, bool bFromHierarchy) const
+{
+	SelectedDetailsPanelItemViewModel.Reset();
+	if(DetailsPanel.IsValid())
+	{
+		if(HierarchyItem.IsValid() && HierarchyItem->AllowEditingInDetailsPanel())
+		{
+			// when se select a section, and the previous item selection is no longer available due to it, we would get a selection refresh next tick
+			// to wipe out the current selection. We want to avoid that, so we manually clear the selected items in that case.
+			if(HierarchyItem->GetData()->IsA<UHierarchySection>())
+			{
+				HierarchyTreeView->ClearSelection();
+			}
+
+			// we clear the selection of the other tree view to avoid 
+			if(bFromHierarchy)
+			{
+				SourceTreeView->ClearSelection();
+			}
+			else
+			{
+				HierarchyTreeView->ClearSelection();
+			}
+			
+			UObject* DataForEditing = HierarchyItem->GetDataForEditing();
+			DataForEditing->SetFlags(RF_Transactional);
+
+			// we make sure the object we are editing is transactional
+			DetailsPanel->SetObject(DataForEditing);
+			SelectedDetailsPanelItemViewModel = HierarchyItem;
+		}
+		else
+		{
+			SelectedDetailsPanelItemViewModel.Reset();
+			DetailsPanel->SetObject(nullptr);
+		}
+	}
+	
+	if(DetailsPanel.IsValid() && OnGenerateCustomDetailsPanelNameWidget.IsBound() && SelectedDetailsPanelItemViewModel.IsValid())
+	{
+		TSharedRef<SWidget> NameWidget = SNew(SHorizontalBox)
+		+ SHorizontalBox::Slot()
+		.AutoWidth()
+		.Padding(2.f)
+		.HAlign(HAlign_Center)
+		.VAlign(VAlign_Center)
+		[
+			SNew(SImage)
+			.Image(FAppStyle::GetBrush("Icons.Lock"))
+			.Visibility(SelectedDetailsPanelItemViewModel.Pin()->IsEditableByUser().bCanPerform ? EVisibility::Collapsed : EVisibility::Visible)
+			.ToolTipText(SelectedDetailsPanelItemViewModel.Pin()->IsEditableByUser().CanPerformMessage)
+		]
+		+ SHorizontalBox::Slot()
+		.AutoWidth()
+		.Padding(2.f)
+		[
+			OnGenerateCustomDetailsPanelNameWidget.Execute(SelectedDetailsPanelItemViewModel.Pin())
+		];
+		
+		DetailsPanel->SetNameAreaCustomContent(NameWidget);
+	}
+}
+
+void SHierarchyEditor::RunSourceSearch()
+{
+	if(!SourceSearchBox->GetText().IsEmpty())
+	{
+		OnSourceSearchTextChanged(SourceSearchBox->GetText());
+	}
+}
+
+void SHierarchyEditor::OnSourceSearchTextChanged(const FText& Text)
+{
+	SourceSearchResults.Empty();
+	FocusedSearchResult.Reset();
+	SourceTreeView->ClearSelection();
+
+	if(!Text.IsEmpty())
+	{
+		FString TextAsString = Text.ToString();;
+
+		TArray<FSearchItem> SearchItems;
+		GenerateSearchItems(SourceRootViewModel.ToSharedRef(), {}, SearchItems);
+		
+		for(const FSearchItem& SearchItem : SearchItems)
+		{
+			for(const FString& SearchTerm : SearchItem.GetEntry()->GetSearchTerms())
+			{
+				if(SearchTerm.Contains(TextAsString))
+				{
+					SourceSearchResults.Add(SearchItem);
+				}
+			}
+		}
+
+		ExpandSourceSearchResults();
+		SelectNextSourceSearchResult();
+	}
+	else
+	{
+		SourceTreeView->ClearExpandedItems();
+	}
+}
+
+void SHierarchyEditor::OnSourceSearchTextCommitted(const FText& Text, ETextCommit::Type CommitType)
+{
+	bool bIsShiftDown = FSlateApplication::Get().GetModifierKeys().IsShiftDown();
+	if(CommitType == ETextCommit::OnEnter)
+	{
+		if(bIsShiftDown == false)
+		{
+			SelectNextSourceSearchResult();
+		}
+		else
+		{
+			SelectPreviousSourceSearchResult();
+		}
+	}
+}
+
+void SHierarchyEditor::OnSearchButtonClicked(SSearchBox::SearchDirection SearchDirection)
+{
+	if(SearchDirection == SSearchBox::Next)
+	{
+		SelectNextSourceSearchResult();
+	}
+	else
+	{
+		SelectPreviousSourceSearchResult();
+	}
+}
+
+void SHierarchyEditor::GenerateSearchItems(TSharedRef<FHierarchyElementViewModel> Root, TArray<TSharedPtr<FHierarchyElementViewModel>> ParentChain, TArray<FSearchItem>& OutSearchItems)
+{
+	const TArray<TSharedPtr<FHierarchyElementViewModel>> FilteredChildren = Root->GetFilteredChildren();
+	ParentChain.Add(Root);
+	OutSearchItems.Add(FSearchItem{ParentChain});
+	for(TSharedPtr<FHierarchyElementViewModel> Child : FilteredChildren)
+	{
+		GenerateSearchItems(Child.ToSharedRef(), ParentChain, OutSearchItems);
+	}
+}
+
+void SHierarchyEditor::ExpandSourceSearchResults()
+{
+	SourceTreeView->ClearExpandedItems();
+
+	for(FSearchItem& SearchResult : SourceSearchResults)
+	{
+		for(TSharedPtr<FHierarchyElementViewModel>& EntryInPath : SearchResult.Path)
+		{
+			SourceTreeView->SetItemExpansion(EntryInPath, true);
+		}
+	}
+}
+
+void SHierarchyEditor::SelectNextSourceSearchResult()
+{
+	if(SourceSearchResults.IsEmpty())
+	{
+		return;
+	}
+	
+	if(!FocusedSearchResult.IsSet())
+	{
+		FocusedSearchResult = SourceSearchResults[0];
+	}
+	else
+	{
+		int32 CurrentSearchResultIndex = SourceSearchResults.Find(FocusedSearchResult.GetValue());
+		if(SourceSearchResults.IsValidIndex(CurrentSearchResultIndex+1))
+		{
+			FocusedSearchResult = SourceSearchResults[CurrentSearchResultIndex+1];
+		}
+		else
+		{
+			FocusedSearchResult = SourceSearchResults[0];
+		}
+	}
+
+	SourceTreeView->ClearSelection();
+	SourceTreeView->RequestScrollIntoView(FocusedSearchResult.GetValue().GetEntry());
+	SourceTreeView->SetItemSelection(FocusedSearchResult.GetValue().GetEntry(), true);
+}
+
+void SHierarchyEditor::SelectPreviousSourceSearchResult()
+{
+	if(SourceSearchResults.IsEmpty())
+	{
+		return;
+	}
+	
+	if(!FocusedSearchResult.IsSet())
+	{
+		FocusedSearchResult = SourceSearchResults[0];
+	}
+	else
+	{
+		int32 CurrentSearchResultIndex = SourceSearchResults.Find(FocusedSearchResult.GetValue());
+		if(SourceSearchResults.IsValidIndex(CurrentSearchResultIndex-1))
+		{
+			FocusedSearchResult = SourceSearchResults[CurrentSearchResultIndex-1];
+		}
+		else
+		{
+			FocusedSearchResult = SourceSearchResults[SourceSearchResults.Num()-1];
+		}
+	}
+
+	SourceTreeView->ClearSelection();
+	SourceTreeView->RequestScrollIntoView(FocusedSearchResult.GetValue().GetEntry());
+	SourceTreeView->SetItemSelection(FocusedSearchResult.GetValue().GetEntry(), true);
+}
+
+TOptional<SSearchBox::FSearchResultData> SHierarchyEditor::GetSearchResultData() const
+{
+	if(SourceSearchResults.Num() > 0)
+	{
+		SSearchBox::FSearchResultData SearchResultData;
+		SearchResultData.NumSearchResults = SourceSearchResults.Num();
+
+		if(FocusedSearchResult.IsSet())
+		{
+			// we add one just to make it look nicer as this is merely for cosmetic purposes
+			SearchResultData.CurrentSearchResultIndex = SourceSearchResults.Find(FocusedSearchResult.GetValue()) + 1;
+		}
+		else
+		{
+			SearchResultData.CurrentSearchResultIndex = INDEX_NONE;
+		}
+
+		return SearchResultData;
+	}
+
+	return TOptional<SSearchBox::FSearchResultData>();
+}
+
+FReply SHierarchyEditor::HandleHierarchyRootDrop(const FGeometry& Geometry, const FDragDropEvent& DragDropEvent) const
+{
+	if(TSharedPtr<FHierarchyDragDropOp> DragDropOp = DragDropEvent.GetOperationAs<FHierarchyDragDropOp>())
+	{
+		HierarchyViewModel->GetHierarchyRootViewModel()->OnDroppedOn(DragDropOp->GetDraggedElement().Pin(), EItemDropZone::OntoItem);
+		return FReply::Handled();
+	}
+
+	return FReply::Unhandled();
+}
+
+FHierarchyElementViewModel::FCanPerformActionResults SHierarchyEditor::CanDropOnRoot(TSharedPtr<FHierarchyElementViewModel> DraggedItem) const
+{
+	return HierarchyViewModel->GetHierarchyRootViewModel()->CanDropOnInternal(DraggedItem, EItemDropZone::OntoItem);
+}
+
+bool SHierarchyEditor::OnCanDropOnRoot(TSharedPtr<FDragDropOperation> DragDropOperation) const
+{
+	if(DragDropOperation->IsOfType<FHierarchyDragDropOp>())
+	{
+		TSharedPtr<FHierarchyDragDropOp> HierarchyDragDropOp = StaticCastSharedPtr<FHierarchyDragDropOp>(DragDropOperation);
+		TSharedPtr<FHierarchyElementViewModel> DraggedItem = HierarchyDragDropOp->GetDraggedElement().Pin();
+		return CanDropOnRoot(DraggedItem).bCanPerform;
+	}
+	
+	return false;
+}
+
+void SHierarchyEditor::OnRootDragEnter(const FDragDropEvent& DragDropEvent) const
+{
+	if(TSharedPtr<FHierarchyDragDropOp> HierarchyDragDropOp = DragDropEvent.GetOperationAs<FHierarchyDragDropOp>())
+	{
+		FHierarchyElementViewModel::FCanPerformActionResults CanPerformActionResults = CanDropOnRoot(HierarchyDragDropOp->GetDraggedElement().Pin());
+		HierarchyDragDropOp->SetDescription(CanPerformActionResults.CanPerformMessage);
+	}
+}
+
+void SHierarchyEditor::OnRootDragLeave(const FDragDropEvent& DragDropEvent) const
+{
+	if(TSharedPtr<FHierarchyDragDropOp> HierarchyDragDropOp = DragDropEvent.GetOperationAs<FHierarchyDragDropOp>())
+	{
+		HierarchyDragDropOp->SetDescription(FText::GetEmpty());
+	}
+}
+
+FSlateColor SHierarchyEditor::GetRootIconColor() const
+{
+	if(FSlateApplication::Get().IsDragDropping())
+	{
+		if(FSlateApplication::Get().GetDragDroppingContent()->IsOfType<FHierarchyDragDropOp>())
+		{
+			TSharedPtr<FHierarchyDragDropOp> HierarchyDragDropOp = StaticCastSharedPtr<FHierarchyDragDropOp>(FSlateApplication::Get().GetDragDroppingContent());
+			if(CanDropOnRoot(HierarchyDragDropOp->GetDraggedElement().Pin()).bCanPerform)
+			{
+				return FLinearColor(0.8f, 0.8f, 0.8f, 0.8f);
+			}
+		}
+	}
+
+	return FLinearColor(0.2f, 0.2f, 0.2f, 0.5f);
+}
+
+FString SHierarchyEditor::GetReferencerName() const
+{
+	return TEXT(" Hierarchy");
+}
+
+void SHierarchyEditor::AddReferencedObjects(FReferenceCollector& Collector)
+{
+	Collector.AddReferencedObject(SourceRoot);
+}
+
+void SHierarchyEditor::NotifyPostChange(const FPropertyChangedEvent& PropertyChangedEvent,	FProperty* PropertyThatChanged)
+{
+	HierarchyViewModel->OnHierarchyPropertiesChanged().Broadcast();
+}
+
+void SSectionDragDropTarget::Construct(const FArguments& InArgs,TSharedPtr<FHierarchySectionViewModel> InOwningSection, EItemDropZone InItemDropZone)
+{
+	OwningSection = InOwningSection;
+	DropZone = InItemDropZone;
+	
+	SDropTarget::Construct(InArgs._DropTargetArgs);
+}
+
+FReply SSectionDragDropTarget::OnDragOver(const FGeometry& MyGeometry, const FDragDropEvent& DragDropEvent)
+{
+	if(TSharedPtr<FHierarchyDragDropOp> HierarchyDragDropOp = DragDropEvent.GetOperationAs<FHierarchyDragDropOp>())
+	{
+		TSharedPtr<FHierarchyElementViewModel> DraggedItem = HierarchyDragDropOp->GetDraggedElement().Pin();
+		FHierarchyElementViewModel::FCanPerformActionResults Results = OwningSection->CanDropOn(DraggedItem, DropZone);
+		HierarchyDragDropOp->SetDescription(Results.CanPerformMessage);
+	}
+
+	return SDropTarget::OnDragOver(MyGeometry, DragDropEvent);
+}
+
+void SSectionDragDropTarget::OnDragLeave(const FDragDropEvent& DragDropEvent)
+{
+	TSharedPtr<FHierarchyDragDropOp> HierarchyDragDropOp = DragDropEvent.GetOperationAs<FHierarchyDragDropOp>();
+
+	if(HierarchyDragDropOp.IsValid())
+	{
+		HierarchyDragDropOp->SetDescription(FText::GetEmpty());
+	}
+	
+	SDropTarget::OnDragLeave(DragDropEvent);
+}
+
+#undef LOCTEXT_NAMESPACE
diff --git a/Engine/Source/Editor/HierarchyEditor/Public/HierarchyEditorCommands.h b/Engine/Source/Editor/HierarchyEditor/Public/HierarchyEditorCommands.h
new file mode 100644
--- /dev/null
+++ b/Engine/Source/Editor/HierarchyEditor/Public/HierarchyEditorCommands.h
@@ -0,0 +1,25 @@
+// Copyright Epic Games, Inc. All Rights Reserved.
+
+#pragma once
+
+#include "CoreMinimal.h"
+#include "Styling/CoreStyle.h"
+#include "Framework/Commands/Commands.h"
+
+class HIERARCHYEDITOR_API FHierarchyEditorCommands : public TCommands<FHierarchyEditorCommands>
+{
+public:
+	
+	FHierarchyEditorCommands()
+		: TCommands<FHierarchyEditorCommands>( TEXT("HierarchyEditorCommands"), NSLOCTEXT("HierarchyEditorCommands", "Hierarchy Editor Commands", "Hierarchy Editor Commands"), NAME_None, "HierarchyEditorStyle" )
+	{
+	}
+
+	virtual ~FHierarchyEditorCommands() override
+	{
+	}
+
+	virtual void RegisterCommands() override;	
+
+	TSharedPtr<FUICommandInfo> FindInHierarchy;
+};
diff --git a/Engine/Source/Editor/HierarchyEditor/Public/HierarchyEditorMisc.h b/Engine/Source/Editor/HierarchyEditor/Public/HierarchyEditorMisc.h
new file mode 100644
--- /dev/null
+++ b/Engine/Source/Editor/HierarchyEditor/Public/HierarchyEditorMisc.h
@@ -0,0 +1,12 @@
+﻿// Copyright Epic Games, Inc. All Rights Reserved.
+
+#pragma once
+
+#include "Containers/Set.h"
+#include "UObject/NameTypes.h"
+
+
+namespace UE::HierarchyEditor
+{
+	HIERARCHYEDITOR_API FName GetUniqueName(FName CandidateName, const TSet<FName>& ExistingNames);
+}
diff --git a/Engine/Source/Editor/HierarchyEditor/Public/HierarchyEditorModule.h b/Engine/Source/Editor/HierarchyEditor/Public/HierarchyEditorModule.h
new file mode 100644
--- /dev/null
+++ b/Engine/Source/Editor/HierarchyEditor/Public/HierarchyEditorModule.h
@@ -0,0 +1,15 @@
+﻿// Copyright Epic Games, Inc. All Rights Reserved.
+
+#pragma once
+
+#include "CoreMinimal.h"
+#include "Modules/ModuleManager.h"
+
+DECLARE_LOG_CATEGORY_EXTERN(LogHierarchyEditor, Log, All)
+
+class FHierarchyEditorModule : public IModuleInterface
+{
+public:
+    virtual void StartupModule() override;
+    virtual void ShutdownModule() override;
+};
diff --git a/Engine/Source/Editor/HierarchyEditor/Public/HierarchyEditorStyle.h b/Engine/Source/Editor/HierarchyEditor/Public/HierarchyEditorStyle.h
new file mode 100644
--- /dev/null
+++ b/Engine/Source/Editor/HierarchyEditor/Public/HierarchyEditorStyle.h
@@ -0,0 +1,29 @@
+// Copyright Epic Games, Inc. All Rights Reserved.
+
+#pragma once
+
+#include "Styling/SlateStyle.h"
+
+/** Manages the style which provides resources for Hierarchy Editor widgets. */
+class HIERARCHYEDITOR_API FHierarchyEditorStyle : public FSlateStyleSet
+{
+public:
+	static void Register();
+	static void Unregister();
+	static void Shutdown();
+
+	/** reloads textures used by slate renderer */
+	static void ReloadTextures();
+
+	/** @return The Slate style set for Hierarchy Editor widgets */
+	static const FHierarchyEditorStyle& Get();
+
+	static void ReinitializeStyle();
+
+private:	
+	FHierarchyEditorStyle();
+
+	void InitHierarchyEditor();
+
+	static TSharedPtr<FHierarchyEditorStyle> HierarchyEditorStyle;
+};
