New Tool: Hierarchy Editor

diff --git a/Engine/Source/Editor/HierarchyEditor/Public/HierarchyViewModelBase.h b/Engine/Source/Editor/HierarchyEditor/Public/HierarchyViewModelBase.h
new file mode 100644
--- /dev/null
+++ b/Engine/Source/Editor/HierarchyEditor/Public/HierarchyViewModelBase.h
@@ -0,0 +1,959 @@
+// Copyright Epic Games, Inc. All Rights Reserved.
+
+#pragma once
+
+#include "CoreMinimal.h"
+#include "PropertyEditorDelegates.h"
+#include "ToolMenuSection.h"
+#include "IPropertyRowGenerator.h"
+#include "Misc/TransactionObjectEvent.h"
+#include "EditorUndoClient.h"
+#include "Widgets/Views/STableRow.h"
+#include "TickableEditorObject.h"
+#include "Engine/TimerHandle.h"
+#include "Templates/SubclassOf.h"
+#include "HierarchyViewModelBase.generated.h"
+
+/** HIERARCHY EDITOR
+ *	The hierarchy editor is a generic tool to organize and structure all kinds of data.
+ *	It inherently supports sections, categories, and items. You can add your own items and customize how they are visualized.
+ *	Each hierarchy element is a UObject, and some hierarchy elements will represent externally owned data.
+ *	For example, categories and sections defined within the hierarchy are also owned by the hierarchy, but an item might represent a parameter defined elsewhere.
+ *
+ *  Each hierarchy element is pure data and should not reference externally owned data that could become invalid.
+ *  To define per-element rules, each hierarchy element gets assigned one view model.
+ *  
+ *	To use the Hierarchy Editor, you need multiple things:
+ *	1) A UHierarchyRoot object that the Hierarchy Editor uses to store the created hierarchy
+ *	2) A UHierarchyViewModelBase-derived object that defines core hierarchy rules. This is the main object responsible for configuring your hierarchy.
+ *	-- The UHierarchyViewModelBase-derived class has multiple virtual functions you need to override. The key functions are:
+ *	--- a) GetHierarchyRoot(), pointing to the UHierarchyRoot object you created in 1)
+ *	--- b) PrepareSourceItems(UHierarchyRoot* SourceRoot, ...), which you need to use to populate the list of elements to be organized.
+ *	------ These are transient items and a transient root. You can also add sections and categories here if you wish to structure the source element view.
+ *	------ To add a child, call NewObject with ParentElement as the outer (in a flat list the SourceRoot).
+ *	------ Then use ParentElement->GetChildrenMutable().Add(NewElement).
+ *	------ This way, you can write your own initialization functions
+ *	------ You can also add categories and sections to the SourceRoot
+ *	--- c) Optionally but likely: CreateCustomViewModelForData(UHierarchyElement* Element, ...), which is used to create and assign non-default view models for each hierarchy element
+ *	------ You don't need to implement this, but any slightly advanced use case that requires modification to hierarchy logic will need to implement this
+ *	3) An SHierarchyEditor widget, which takes in the UHierarchyViewModelBase-derived object you created in 2). This way, the hierarchy editor knows where to store data, and how to initialize.
+ *	--- Additionally, it has arguments such as 'OnGenerateRowContentWidget', which gives you an ElementViewModel.
+ *	--- After determining the type, you can use this to define widgets for each hierarchy element.
+ *
+ *	Tips and tricks:
+ *	1) Each hierarchy element can have a FHierarchyElementIdentity consisting of guid(s) and/or name(s).
+ *	--- This can be used to uniquely identify an element, navigate to it and so on.
+ *	2) To deal with automated cleanup of stale hierarchy elements that represent external data, you can set a UHierarchyDataRefreshContext-derived object on the UHierarchyViewModelBase object.
+ *	--- The idea is to create a new derived class with objects and properties of the system that uses the Hierarchy Editor.
+ *	--- For example, in a graph based system, the DataRefreshContext object can point to the graph.
+ *	--- In an FHierarchyElementViewModel's 'DoesExternalDataStillExist' function,
+ *	--- you can access the DataRefreshContext to determine whether your externally represented hierarchy elements should still exist.
+ *	--- If not, it is automatically deleted.
+ *	3) In the details panel, You can edit the hierarchy elements themselves, or external objects 
+ *
+ *  To make use of the created hierarchy, you access the UHierarchyRoot object you created in 1), and query it for its children, sections etc.
+ *  The HierarchyEditor does not define how to use the created hierarchy data in your own UI. It only lets you structure and edit data.
+ *  How you use it from a 'data consumption' point of view is up to you.
+ */
+
+struct FHierarchyRootViewModel;
+/** This struct is used to identify a given hierarchy element and can be based on guids and/or names.
+ *  This is particularly useful when a hierarchy element represents an object or a property that is not owned by the hierarchy itself.
+ *  
+ *  
+ */
+USTRUCT()
+struct FHierarchyElementIdentity
+{
+	GENERATED_BODY()
+
+	FHierarchyElementIdentity() {}
+	FHierarchyElementIdentity(TArray<FGuid> InGuids, TArray<FName> InNames) : Guids(InGuids), Names(InNames) {}
+	
+	/** An array of guids that have to be satisfied in order to match. */
+	UPROPERTY()
+	TArray<FGuid> Guids;
+
+	/** Optionally, an array of names can be specified in place of guids. If guids & names are present, guids have to be satisfied first, then names. */
+	UPROPERTY()
+	TArray<FName> Names;
+
+	bool IsValid() const
+	{
+		return Guids.Num() > 0 || Names.Num() > 0;
+	}
+	
+	bool operator==(const FHierarchyElementIdentity& OtherIdentity) const
+	{
+		if(Guids.Num() != OtherIdentity.Guids.Num() || Names.Num() != OtherIdentity.Names.Num())
+		{
+			return false;
+		}
+
+		for(int32 GuidIndex = 0; GuidIndex < Guids.Num(); GuidIndex++)
+		{
+			if(Guids[GuidIndex] != OtherIdentity.Guids[GuidIndex])
+			{
+				return false;
+			}
+		}
+
+		for(int32 NameIndex = 0; NameIndex < Names.Num(); NameIndex++)
+		{
+			if(!Names[NameIndex].IsEqual(OtherIdentity.Names[NameIndex]))
+			{
+				return false;
+			}
+		}
+
+		return true;
+	}
+
+	bool operator!=(const FHierarchyElementIdentity& OtherIdentity) const
+	{
+		return !(*this == OtherIdentity);
+	}
+};
+
+FORCEINLINE uint32 GetTypeHash(const FHierarchyElementIdentity& Identity)
+{
+	uint32 Hash = 0;
+	
+	for(const FGuid& Guid : Identity.Guids)
+	{
+		Hash = HashCombine(Hash, GetTypeHash(Guid));
+	}
+	
+	for(const FName& Name : Identity.Names)
+	{
+		Hash = HashCombine(Hash, GetTypeHash(Name));
+	}
+	
+	return Hash;
+}
+
+/** A base class that is used to refresh data that represents external data. Inherit from this class if you need more context data. */
+UCLASS()
+class HIERARCHYEDITOR_API UHierarchyDataRefreshContext : public UObject
+{
+	GENERATED_BODY()
+};
+
+UCLASS()
+class HIERARCHYEDITOR_API UHierarchyElement : public UObject
+{
+	GENERATED_BODY()
+
+public:
+	UHierarchyElement() { Identity.Guids.Add(FGuid::NewGuid()); }
+	virtual ~UHierarchyElement() override {}
+
+	TArray<TObjectPtr<UHierarchyElement>>& GetChildrenMutable() { return Children; }
+	const TArray<UHierarchyElement*>& GetChildren() const;
+
+	template<class ChildClass>
+	ChildClass* AddChild();
+
+	UHierarchyElement* FindChildWithIdentity(FHierarchyElementIdentity ChildIdentity, bool bSearchRecursively = false);
+
+	UHierarchyElement* CopyAndAddItemAsChild(const UHierarchyElement& ItemToCopy);
+	UHierarchyElement* CopyAndAddItemUnderParentIdentity(const UHierarchyElement& ItemToCopy, FHierarchyElementIdentity ParentIdentity);
+	
+	/** Remove a child with a given identity. Can be searched recursively. This function operates under the assumption there will be only one item with a given identity. */
+	bool RemoveChildWithIdentity(FHierarchyElementIdentity ChildIdentity, bool bSearchRecursively = false);
+	
+	template<class ChildClass>
+	bool DoesOneChildExist(bool bRecursive = false) const;
+
+	template<class ChildClass>
+	TArray<ChildClass*> GetChildrenOfType(TArray<ChildClass*>& Out, bool bRecursive = false) const;
+
+	template<class PREDICATE_CLASS>
+	void SortChildren(const PREDICATE_CLASS& Predicate, bool bRecursive = false);
+	
+	virtual FString ToString() const { return GetName(); }
+	FText ToText() const { return FText::FromString(ToString()); }
+	
+	/** An identity can be optionally set to create a mapping from previously existing guids or names to hierarchy items that represent them. */
+	void SetIdentity(FHierarchyElementIdentity InIdentity) { Identity = InIdentity; }
+	FHierarchyElementIdentity GetPersistentIdentity() const { return Identity; }
+	TArray<FHierarchyElementIdentity> GetParentIdentities() const;
+	
+	/** Overridden modify method to also mark all children as modified */
+	virtual bool Modify(bool bAlwaysMarkDirty = true) override;
+
+protected:
+	virtual void PostLoad() override;
+
+	UPROPERTY()
+	TArray<TObjectPtr<UHierarchyElement>> Children;
+
+	UPROPERTY()
+	FHierarchyElementIdentity Identity;
+	
+	/** An optional guid; can be used if hierarchy items represent outside items */
+	UPROPERTY()
+	FGuid Guid_DEPRECATED;
+};
+
+template <class ChildClass>
+ChildClass* UHierarchyElement::AddChild()
+{
+	ChildClass* NewChild = NewObject<ChildClass>(this);
+	GetChildrenMutable().Add(NewChild);
+
+	return NewChild;
+}
+
+template <class ChildClass>
+bool UHierarchyElement::DoesOneChildExist(bool bRecursive) const
+{
+	for(UHierarchyElement* ChildElement : Children)
+	{
+		if(ChildElement->IsA<ChildClass>())
+		{
+			return true;
+		}
+	}
+
+	if(bRecursive)
+	{
+		for(UHierarchyElement* ChildElement : Children)
+		{
+			if(ChildElement->DoesOneChildExist<ChildClass>(bRecursive))
+			{
+				return true;
+			}
+		}
+	}
+
+	return false;
+}
+
+template <class ChildClass>
+TArray<ChildClass*> UHierarchyElement::GetChildrenOfType(TArray<ChildClass*>& Out, bool bRecursive) const
+{
+	for(UHierarchyElement* ChildElement : Children)
+	{
+		if(ChildElement->IsA<ChildClass>())
+		{
+			Out.Add(Cast<ChildClass>(ChildElement));
+		}
+	}
+
+	if(bRecursive)
+	{
+		for(UHierarchyElement* ChildElement : Children)
+		{
+			ChildElement->GetChildrenOfType<ChildClass>(Out, bRecursive);
+		}
+	}
+
+	return Out;
+}
+
+template <class PREDICATE_CLASS>
+void UHierarchyElement::SortChildren(const PREDICATE_CLASS& Predicate, bool bRecursive)
+{
+	Children.Sort(Predicate);
+
+	if(bRecursive)
+	{
+		for(TObjectPtr<UHierarchyElement> ChildElement : Children)
+		{
+			ChildElement->SortChildren(Predicate, bRecursive);
+		}
+	}
+}
+
+/** A minimal implementation of a section. */
+UCLASS()
+class HIERARCHYEDITOR_API UHierarchySection : public UHierarchyElement
+{
+	GENERATED_BODY()
+
+public:
+	UHierarchySection() {}
+
+	void SetSectionName(FName InSectionName) { Section = InSectionName; }
+	FName GetSectionName() const { return Section; }
+	
+	void SetSectionNameAsText(const FText& Text);
+	FText GetSectionNameAsText() const { return FText::FromName(Section); }
+
+	void SetTooltip(const FText& InTooltip) { Tooltip = InTooltip; }
+	FText GetTooltip() const { return Tooltip; }
+
+	virtual FString ToString() const override { return Section.ToString(); }
+private:
+	UPROPERTY()
+	FName Section;
+
+	/** The tooltip used when the user is hovering this section */
+	UPROPERTY(EditAnywhere, Category = "Section", meta = (MultiLine = "true"))
+	FText Tooltip;
+};
+
+/** UHierarchyRoot is used as the main object for serialization purposes, and a transient root is created automatically by the widget to populate the source list of items. */
+UCLASS()
+class HIERARCHYEDITOR_API UHierarchyRoot : public UHierarchyElement
+{
+	GENERATED_BODY()
+public:
+	UHierarchyRoot() {}
+	virtual ~UHierarchyRoot() override {}
+
+	const TArray<UHierarchySection*>& GetSectionData() const { return Sections; }
+	TArray<TObjectPtr<UHierarchySection>>& GetSectionDataMutable() { return Sections; }
+
+	TSet<FName> GetSections() const;
+	int32 GetSectionIndex(UHierarchySection* Section) const;
+
+	UHierarchySection* AddSection(FText InNewSectionName, int32 InsertIndex = INDEX_NONE, TSubclassOf<UHierarchySection> SectionClass = UHierarchySection::StaticClass());
+	UHierarchySection* FindSectionByIdentity(FHierarchyElementIdentity SectionIdentity);
+	/** This will copy the section element itself */
+	void DuplicateSectionFromOtherRoot(const UHierarchySection& SectionToCopy);
+	void RemoveSection(FText SectionName);
+	void RemoveSectionByIdentity(FHierarchyElementIdentity SectionIdentity);
+	
+	virtual bool Modify(bool bAlwaysMarkDirty = true) override;
+
+	virtual void EmptyAllData();
+protected:
+
+	UPROPERTY()
+	TArray<TObjectPtr<UHierarchySection>> Sections;
+};
+
+/** A minimal implementation of an item. Inherit from this and add your own properties. */
+UCLASS(MinimalAPI)
+class UHierarchyItem : public UHierarchyElement
+{
+	GENERATED_BODY()
+public:
+	UHierarchyItem() {}
+	virtual ~UHierarchyItem() override {}
+};
+
+/** A category, potentially pointing at the section it belongs to. Only top-level categories can belong to sections by default.
+ *  Inherit from this to add your own properties.  */
+UCLASS()
+class HIERARCHYEDITOR_API UHierarchyCategory : public UHierarchyElement
+{
+	GENERATED_BODY()
+public:
+	UHierarchyCategory() {}
+	UHierarchyCategory(FName InCategory) : Category(InCategory) {}
+	
+	void SetCategoryName(FName NewCategory) { Category = NewCategory; }
+	FName GetCategoryName() const { return Category; }
+
+	FText GetCategoryAsText() const { return FText::FromName(Category); }
+	FText GetTooltip() const { return Tooltip; }
+
+	void SetSection(UHierarchySection* InSection) { Section = InSection; }
+	const UHierarchySection* GetSection() const { return Section; }
+
+	virtual FString ToString() const override { return Category.ToString(); }
+
+	/** Since the category points to a section object, during merge or copy paste etc. it is possible the section pointer will point at a section from another root.
+	 *  We fix this up by looking through our available sections and match up via persistent identity.
+	 *  This function expects the correct section with the same identity to exist already at the root level
+	 */
+	void FixupSectionLinkage();
+private:
+	UPROPERTY()
+	FName Category;
+
+	/** The tooltip used when the user is hovering this category */
+	UPROPERTY(EditAnywhere, Category = "Category", meta = (MultiLine = "true"))
+	FText Tooltip;
+
+	UPROPERTY()
+	TObjectPtr<UHierarchySection> Section = nullptr;
+};
+
+struct FHierarchyElementViewModel;
+struct FHierarchySectionViewModel;
+
+/** Inherit from this to allow UI customization for your drag & drop operation by overriding CreateCustomDecorator. */
+class HIERARCHYEDITOR_API FHierarchyDragDropOp : public FDragDropOperation
+{
+public:
+	DRAG_DROP_OPERATOR_TYPE(FHierarchyDragDropOp, FDragDropOperation)
+
+	FHierarchyDragDropOp(TSharedPtr<FHierarchyElementViewModel> InDraggedElementViewModel);
+
+	virtual void Construct() override { FDragDropOperation::Construct(); }
+	virtual TSharedPtr<SWidget> GetDefaultDecorator() const override final;
+	
+	/** Override this custom decorator function to provide custom widget visuals. If not specified, you can still use Label and Description. */
+	virtual TSharedRef<SWidget> CreateCustomDecorator() const { return SNullWidget::NullWidget; }
+	
+	TWeakPtr<FHierarchyElementViewModel> GetDraggedElement() { return DraggedElement; }
+	
+	void SetLabel(FText InText) { Label = InText; }
+	FText GetLabel() const { return Label; }
+
+	void SetDescription(FText InText) { Description = InText; }
+	FText GetDescription() const { return Description; }
+
+	void SetFromSourceList(bool bInFromSourceList) { bFromSourceList = bInFromSourceList; }
+	bool GetIsFromSourceList() const { return bFromSourceList; }
+protected:
+	TWeakPtr<FHierarchyElementViewModel> DraggedElement;
+	
+	/** Label will be displayed if no custom decorator has been specified. */
+	FText Label;
+	/** Useful for runtime tweaking of the tooltip based on what we are hovering. Always displayed if not-empty */
+	FText Description;
+	/** If the drag drop op is from the source list, we can further customize the actions */
+	bool bFromSourceList = false;
+};
+
+UCLASS(BlueprintType)
+class HIERARCHYEDITOR_API UHierarchyMenuContext : public UObject
+{
+	GENERATED_BODY()
+
+public:
+	TArray<TSharedPtr<FHierarchyElementViewModel>> MenuHierarchyElements;
+};
+
+/** The main controller class for the SHierarchyEditor widget. Defines core hierarchy rules.
+ *  Inherit from this and override the required virtual functions, instantiate an object, Initialize it and pass it to the SHierarchyEditor widget. */
+UCLASS(Abstract, Transient)
+class HIERARCHYEDITOR_API UHierarchyViewModelBase : public UObject, public FSelfRegisteringEditorUndoClient
+{
+public:
+	DECLARE_MULTICAST_DELEGATE(FOnHierarchyChanged)
+	DECLARE_MULTICAST_DELEGATE(FOnHierarchyPropertiesChanged)
+	DECLARE_DELEGATE_OneParam(FOnSectionActivated, TSharedPtr<FHierarchySectionViewModel> Section)
+	DECLARE_DELEGATE_OneParam(FOnElementAdded, TSharedPtr<FHierarchyElementViewModel> AddedElement)
+	DECLARE_DELEGATE_OneParam(FOnRefreshViewRequested, bool bForceFullRefresh)
+	DECLARE_DELEGATE_OneParam(FOnNavigateToElementIdentityInHierarchyRequested, FHierarchyElementIdentity Identity)
+	DECLARE_DELEGATE_OneParam(FOnNavigateToElementInHierarchyRequested, TSharedPtr<FHierarchyElementViewModel> ElementViewModel)
+
+	GENERATED_BODY()
+
+	UHierarchyViewModelBase();
+	virtual ~UHierarchyViewModelBase() override;
+
+	/** Initialize is called automatically for you, but it is recommended to call it manually after creating the HierarchyViewModel in your own Initialize function. This lets you access external data. */
+	void Initialize();
+	/** Call Finalize manually when you no longer need the HierarchyViewModel. */
+	void Finalize();
+	
+	bool IsInitialized() const { return bIsInitialized; }
+	bool IsFinalized() const { return bIsFinalized; }
+	bool IsValid() const { return IsInitialized() && !IsFinalized(); }
+
+	/** Creates view model hierarchy elements. To create custom view models, override CreateCustomViewModelForElement. */
+	TSharedPtr<FHierarchyElementViewModel> CreateViewModelForElement(UHierarchyElement* Element, TSharedPtr<FHierarchyElementViewModel> Parent);
+	/** Get the root view model associated with the hierarchy. */
+	TSharedPtr<struct FHierarchyRootViewModel> GetHierarchyRootViewModel() const { return HierarchyRootViewModel; }
+	/** Hierarchy items reflect the already edited hierarchy. This should generally be constructed from persistent serialized data. */
+	const TArray<TSharedPtr<FHierarchyElementViewModel>>& GetHierarchyItems() const;
+
+public:
+	/** The hierarchy root the widget is editing. This should point to persistent data stored somewhere else as the serialized root of the hierarchy. */
+	virtual UHierarchyRoot* GetHierarchyRoot() const PURE_VIRTUAL(UHierarchyViewModelBase::GetHierarchyRoot, return nullptr;);
+	/** Prepares the items we want to create a hierarchy for. Primary purpose is to add children to the source root to gather the items to display in the source panel.
+	 * The root view model is also given as a way to forcefully sync view models to access additional functionality, if needed */
+	virtual void PrepareSourceItems(UHierarchyRoot* SourceRoot, TSharedPtr<FHierarchyRootViewModel> SourceRootViewModel) PURE_VIRTUAL(UHierarchyViewModelBase::PrepareSourceItems,);
+	/** The class used for creating categories. You can subclass UHierarchyCategory to add new properties. */
+	virtual TSubclassOf<UHierarchyCategory> GetCategoryDataClass() const;
+	/** The class used for creating sections. You can subclass UHierarchySection to add new properties. */
+	virtual TSubclassOf<UHierarchySection> GetSectionDataClass() const;
+	/** Function to implement drag drop ops. FHierarchyDragDropOp is a default implementation for a single hierarchy element. Inherit from it and override CreateCustomDecorator for custom UI. */
+	virtual TSharedRef<FHierarchyDragDropOp> CreateDragDropOp(TSharedRef<FHierarchyElementViewModel> Item);
+	/** This needs to return true if you want the details panel to show up. */
+	virtual bool SupportsDetailsPanel() { return true; }
+	/** Overriding this allows to define details panel instance customizations for specific UClasses */
+	virtual TArray<TTuple<UClass*, FOnGetDetailCustomizationInstance>> GetInstanceCustomizations() { return {}; }
+
+protected:
+	/** Additional commands can be specified overriding the SetupCommands function. */
+	virtual void SetupCommands() {}
+private:
+	/** Lets you add some additional logic to the Initialize function. */
+	virtual void InitializeInternal() {}
+	virtual void FinalizeInternal() {}
+
+	/** This function is used to determine custom view models for Hierarchy Elements. Called by CreateViewModelForElement. */
+	virtual TSharedPtr<FHierarchyElementViewModel> CreateCustomViewModelForElement(UHierarchyElement* Element, TSharedPtr<FHierarchyElementViewModel> Parent) { return nullptr; }
+	/** Override this to add custom context menu options. The items to generate this for are contained in a UHierarchyMenuContext object accessible through the ToolMenu. */
+	virtual void GenerateDynamicContextMenuInternal(TArray<TSharedPtr<FHierarchyElementViewModel>> ViewModels) const {}
+public:
+	virtual void GenerateDynamicContextMenu(UToolMenu* ToolMenu) const;
+
+	const UHierarchyDataRefreshContext* GetRefreshContext() const { return RefreshContext; }
+	/** Set the refresh context to easily allow Hierarchy Elements representing external data to access whether the external data still exists. */
+	void SetRefreshContext(UHierarchyDataRefreshContext* InContext) { RefreshContext = InContext; }
+	
+	UHierarchyElement* AddElementUnderRoot(TSubclassOf<UHierarchyElement> NewChildClass, FHierarchyElementIdentity ChildIdentity);
+	void AddCategory(TSharedPtr<FHierarchyElementViewModel> CategoryParent = nullptr) const;
+	void AddSection() const;
+
+	void DeleteElementWithIdentity(FHierarchyElementIdentity Identity);
+	void DeleteElementsWithIdentities(TArray<FHierarchyElementIdentity> Identities);
+	void DeleteElements(TArray<TSharedPtr<FHierarchyElementViewModel>> ViewModels) const;
+
+	void NavigateToElementInHierarchy(const FHierarchyElementIdentity& HierarchyIdentity) const;
+	void NavigateToElementInHierarchy(const TSharedRef<FHierarchyElementViewModel> HierarchyElement) const;
+	
+	/** Refreshes all data and widgets */
+	void ForceFullRefresh();
+	void ForceFullRefreshOnTimer();
+	void RequestFullRefreshNextFrame();
+	
+	TSharedRef<FUICommandList> GetCommands() const { return Commands.ToSharedRef(); }
+	
+	void OnGetChildren(TSharedPtr<FHierarchyElementViewModel> Element, TArray<TSharedPtr<FHierarchyElementViewModel>>& OutChildren) const;
+	
+	void RefreshAllViews(bool bFullRefresh = false) const;
+	void RefreshSourceView(bool bFullRefresh = false) const;
+	void RefreshHierarchyView(bool bFullRefresh = false) const;
+	void RefreshSectionsView() const;
+	
+	// Delegate that call functions from SHierarchyEditor
+	FSimpleDelegate& OnRefreshSourceItemsRequested() { return RefreshSourceItemsRequestedDelegate; }
+	FOnRefreshViewRequested& OnRefreshSourceView() { return RefreshSourceViewDelegate; }
+	FOnRefreshViewRequested& OnRefreshHierarchyView() { return RefreshHierarchyWidgetDelegate; }
+	FSimpleDelegate& OnRefreshSectionsView() { return RefreshSectionsViewDelegate; }
+
+	// Delegates for external systems
+	FOnHierarchyChanged& OnHierarchyChanged() { return OnHierarchyChangedDelegate; } 
+	FOnHierarchyChanged& OnHierarchyPropertiesChanged() { return OnHierarchyPropertiesChangedDelegate; } 
+	FOnElementAdded& OnElementAdded() { return OnElementAddedDelegate; }
+	FOnRefreshViewRequested& OnRefreshViewRequested() { return RefreshAllViewsRequestedDelegate; }
+	FOnNavigateToElementIdentityInHierarchyRequested& OnNavigateToElementIdentityInHierarchyRequested() { return OnNavigateToElementIdentityInHierarchyRequestedDelegate; }
+	FOnNavigateToElementInHierarchyRequested& OnNavigateToElementInHierarchyRequested() { return OnNavigateToElementInHierarchyRequestedDelegate; }
+	FSimpleDelegate& OnInitialized() { return OnInitializedDelegate; }
+	
+	// Sections
+	void SetActiveHierarchySection(TSharedPtr<struct FHierarchySectionViewModel>);
+	TSharedPtr<FHierarchySectionViewModel> GetActiveHierarchySection() const;
+	UHierarchySection* GetActiveHierarchySectionData() const;
+	bool IsHierarchySectionActive(const UHierarchySection* Section) const;
+	FOnSectionActivated& OnHierarchySectionActivated() { return OnHierarchySectionActivatedDelegate; }
+
+	FString OnElementToStringDebug(TSharedPtr<FHierarchyElementViewModel> ElementViewModel) const;
+
+protected:	
+	virtual void PostUndo(bool bSuccess) override;
+	virtual void PostRedo(bool bSuccess) override;
+	virtual bool MatchesContext(const FTransactionContext& InContext, const TArray<TPair<UObject*, FTransactionObjectEvent>>& TransactionObjectContexts) const override;
+private:
+	bool FilterForHierarchySection(TSharedPtr<const FHierarchyElementViewModel> ViewModel) const;
+	bool FilterForUncategorizedRootItemsInAllSection(TSharedPtr<const FHierarchyElementViewModel> ViewModel) const;
+
+	void ToolMenuRequestRename(const FToolMenuContext& Context) const;
+	bool ToolMenuCanRequestRename(const FToolMenuContext& Context) const;
+
+	void ToolMenuDelete(const FToolMenuContext& Context) const;
+	bool ToolMenuCanDelete(const FToolMenuContext& Context) const;
+
+	void ToolMenuNavigateTo(const FToolMenuContext& Context) const;
+	bool ToolMenuCanNavigateTo(const FToolMenuContext& Context) const;
+protected:
+	UPROPERTY()
+	TObjectPtr<UHierarchyRoot> HierarchyRoot;
+	
+	TSharedPtr<struct FHierarchyRootViewModel> HierarchyRootViewModel;
+
+	TWeakPtr<struct FHierarchySectionViewModel> ActiveHierarchySection;
+
+	TSharedPtr<FUICommandList> Commands;
+
+	UPROPERTY(Transient)
+	TObjectPtr<UHierarchyDataRefreshContext> RefreshContext = nullptr;
+	
+protected:
+	// delegate collection to call UI functions
+	FSimpleDelegate RefreshSourceItemsRequestedDelegate;
+	FOnRefreshViewRequested RefreshAllViewsRequestedDelegate;
+	FOnRefreshViewRequested RefreshSourceViewDelegate;
+	FOnRefreshViewRequested RefreshHierarchyWidgetDelegate;
+	FSimpleDelegate RefreshSectionsViewDelegate;
+	FOnNavigateToElementIdentityInHierarchyRequested OnNavigateToElementIdentityInHierarchyRequestedDelegate;
+	FOnNavigateToElementInHierarchyRequested OnNavigateToElementInHierarchyRequestedDelegate;
+	
+	FOnElementAdded OnElementAddedDelegate;
+	FOnSectionActivated OnHierarchySectionActivatedDelegate;
+	FOnSectionActivated OnSourceSectionActivatedDelegate;
+	FOnHierarchyChanged OnHierarchyChangedDelegate;
+	FOnHierarchyPropertiesChanged OnHierarchyPropertiesChangedDelegate;
+	
+	FSimpleDelegate OnInitializedDelegate;
+
+	FTimerHandle FullRefreshNextFrameHandle;
+
+private:
+	UPROPERTY(Transient)
+	bool bIsInitialized = false;
+	
+	UPROPERTY(Transient)
+	bool bIsFinalized = false;
+};
+
+/** The base view model for all elements in the hierarchy. There are four base view models inheriting from this; for roots, items, categories, and sections.
+ *  When creating a new view model, you should inherit from one of those four base view models.
+ */
+struct HIERARCHYEDITOR_API FHierarchyElementViewModel : TSharedFromThis<FHierarchyElementViewModel>, public FTickableEditorObject
+{
+	DECLARE_DELEGATE(FOnSynced)
+	DECLARE_DELEGATE_RetVal_OneParam(bool, FOnFilterChild, const TSharedPtr<const FHierarchyElementViewModel> Child);
+	DECLARE_DELEGATE_OneParam(FOnChildRequestedDeletion, TSharedPtr<FHierarchyElementViewModel> Child)
+	
+	FHierarchyElementViewModel(UHierarchyElement* InElement, TSharedPtr<FHierarchyElementViewModel> InParent, TWeakObjectPtr<UHierarchyViewModelBase> InHierarchyViewModel, bool bInIsForHierarchy)
+		: Element(InElement)
+		, Parent(InParent)
+		, HierarchyViewModel(InHierarchyViewModel)
+		, bIsForHierarchy(bInIsForHierarchy)
+	{		
+		
+	}
+
+	/** Can be implemented for additional logic that the constructor isn't valid for. */
+	virtual void Initialize() {}
+	
+	virtual ~FHierarchyElementViewModel() override;
+	
+	UHierarchyElement* GetDataMutable() const { return Element; }
+	const UHierarchyElement* GetData() const { return Element; }
+	
+	template<class T>
+	T* GetDataMutable() const { return Cast<T>(Element); }
+	
+	template<class T>
+	const T* GetData() const { return Cast<T>(Element); }
+	
+	virtual void Tick(float DeltaTime) override;
+	virtual TStatId GetStatId() const override;
+	
+	virtual FString ToString() const { return Element->ToString(); }
+	FText ToStringAsText() const { return FText::FromString(ToString()); }
+	virtual TArray<FString> GetSearchTerms() const { return {ToString()} ;}
+
+	void RefreshChildrenData();
+	void SyncViewModelsToData();
+	
+	/** Every item view model can define its own sort order for its children. By default we put categories above items. */
+	virtual void SortChildrenData() const;
+	
+	const TArray<TSharedPtr<FHierarchyElementViewModel>>& GetChildren() const { return Children; }
+	TArray<TSharedPtr<FHierarchyElementViewModel>>& GetChildrenMutable() { return Children; }
+	const TArray<TSharedPtr<FHierarchyElementViewModel>>& GetFilteredChildren() const;
+	
+	void AddChildFilter(FOnFilterChild InFilterChild);
+	
+	template<class DataClass, class ViewModelChildClass>
+	void GetChildrenViewModelsForType(TArray<TSharedPtr<ViewModelChildClass>>& OutChildren, bool bRecursive = false);
+
+	/** Returns the hierarchy depth via number of parents above. */
+	int32 GetHierarchyDepth() const;
+	bool HasParent(TSharedPtr<FHierarchyElementViewModel> ParentCandidate, bool bRecursive = false) const;
+	
+	TSharedRef<FHierarchyElementViewModel> DuplicateToThis(TSharedPtr<FHierarchyElementViewModel> ItemToDuplicate, int32 InsertIndex = INDEX_NONE);
+	TSharedRef<FHierarchyElementViewModel> ReparentToThis(TSharedPtr<FHierarchyElementViewModel> ItemToMove, int32 InsertIndex = INDEX_NONE);
+
+	TSharedPtr<FHierarchyElementViewModel> FindViewModelForChild(UHierarchyElement* Child, bool bSearchRecursively = false) const;
+	TSharedPtr<FHierarchyElementViewModel> FindViewModelForChild(FHierarchyElementIdentity ChildIdentity, bool bSearchRecursively = false) const;
+	int32 FindIndexOfChild(UHierarchyElement* Child) const;
+	int32 FindIndexOfDataChild(TSharedPtr<FHierarchyElementViewModel> Child) const;
+	int32 FindIndexOfDataChild(UHierarchyElement* Child) const;
+
+	UHierarchyElement* AddChild(TSubclassOf<UHierarchyElement> NewChildClass, FHierarchyElementIdentity ChildIdentity);
+	
+	/** Deleting will ask the parent to delete its child */
+	void Delete();
+	void DeleteChild(TSharedPtr<FHierarchyElementViewModel> Child);
+
+	TWeakObjectPtr<UHierarchyViewModelBase> GetHierarchyViewModel() const { return HierarchyViewModel; }
+
+	/** Returns a set result if the item can accept a drop either above/onto/below the item.  */
+	TOptional<EItemDropZone> OnCanRowAcceptDrop(const FDragDropEvent& DragDropEvent, EItemDropZone ItemDropZone, TSharedPtr<FHierarchyElementViewModel> Item);
+	virtual FReply OnDroppedOnRow(const FDragDropEvent& DragDropEvent, EItemDropZone ItemDropZone, TSharedPtr<FHierarchyElementViewModel> Item);
+	void OnRowDragLeave(const FDragDropEvent& DragDropEvent);
+
+	struct FCanPerformActionResults
+	{
+		FCanPerformActionResults(bool bInCanPerform) : bCanPerform(bInCanPerform) {}
+		
+		bool bCanPerform = false;
+		/** A message that is used when bCanPerform is false. Will either be used in tooltips in the hierarchy editor or as popup message. */
+		FText CanPerformMessage;
+
+		bool operator==(const bool& bOther) const
+		{
+			return bCanPerform == bOther;
+		}
+
+		bool operator!=(const bool& bOther) const
+		{
+			return !(*this==bOther);
+		}
+	};
+
+	/** Should return true if properties are supposed to be editable & needs to be true if typical operations should work on it (renaming, dragging, deleting etc.) */
+	virtual FCanPerformActionResults IsEditableByUser() { return FCanPerformActionResults(false); }
+	
+	/** Needs to be true in order to allow drag & drop operations to parent items to this item */
+	virtual bool CanHaveChildren() const { return false; }
+	
+	/** Should return true if an item should be draggable. An uneditable item can not be dragged even if CanDragInternal returns true. */
+	FCanPerformActionResults CanDrag();
+	
+	/** Returns true if renamable */
+	bool CanRename() { return IsEditableByUser().bCanPerform && CanRenameInternal(); }
+
+	void Rename(FName NewName) { RenameInternal(NewName); HierarchyViewModel->OnHierarchyPropertiesChanged().Broadcast(); }
+
+	void RequestRename()
+	{
+		if(CanRename() && OnRequestRenameDelegate.IsBound())
+		{
+			bRenamePending = false;
+			OnRequestRenameDelegate.Execute();
+		}
+	}
+
+	void RequestRenamePending()
+	{
+		if(CanRename())
+		{
+			bRenamePending = true;
+		}
+	}
+	
+	/** Returns true if deletable */
+	bool CanDelete() { return IsEditableByUser().bCanPerform && CanDeleteInternal(); }
+
+	/** Returns true if the given item can be dropped on the given target area. */
+	FCanPerformActionResults CanDropOn(TSharedPtr<FHierarchyElementViewModel> DraggedItem, EItemDropZone ItemDropZone) { return CanDropOnInternal(DraggedItem, ItemDropZone); }
+
+	/** Gets executed when an item was dropped on this. */
+	void OnDroppedOn(TSharedPtr<FHierarchyElementViewModel> DroppedItem, EItemDropZone ItemDropZone) { OnDroppedOnInternal(DroppedItem, ItemDropZone); }
+
+	/** Determines the section this item belongs to. */
+	const UHierarchySection* GetSection() const { return GetSectionInternal(); }
+	
+	/** For data cleanup that represents external data, this needs to return true in order for live cleanup to work. */
+	virtual bool RepresentsExternalData() const { return false; }
+	/** This function determines whether a hierarchy item that represents that external data should be maintained during data refresh
+	 * Needs to be implemented if RepresentsExternalData return true.
+	 * The context object can be used to add arbitrary data. */
+	virtual bool DoesExternalDataStillExist(const UHierarchyDataRefreshContext* Context) const { return false; }
+
+	/** The UObject we display in the details panel when this item is selected. By default it's the hierarchy element the view model represents. */
+	virtual UObject* GetDataForEditing() { return Element; }
+	/** Source items are transient, which is why we don't allow editing by default.
+	 * This is useful to override if source data points at actual data to edit. */
+	virtual bool AllowEditingInDetailsPanel() const { return bIsForHierarchy; }
+
+	bool IsForHierarchy() const { return bIsForHierarchy; }
+
+	/** Override this to register dynamic context menu entries when right clicking a single hierarchy item */
+	virtual void AppendDynamicContextMenuForSingleElement(UToolMenu* ToolMenu) {}
+	
+	FReply OnDragDetected(const FGeometry& Geometry, const FPointerEvent& PointerEvent, bool bIsSource);
+
+	FSimpleDelegate& GetOnRequestRename() { return OnRequestRenameDelegate; }
+	FOnSynced& GetOnSynced() { return OnSyncedDelegate; }
+	FOnChildRequestedDeletion& OnChildRequestedDeletion() { return OnChildRequestedDeletionDelegate; }
+
+	TWeakPtr<FHierarchyElementViewModel> GetParent() { return Parent; }
+
+protected:
+	/** Should return true if draggable. An optional message can be provided if false that will show as a slate notification. */
+	virtual FCanPerformActionResults CanDragInternal() { return false; }
+
+	/** Should return true if renamable */
+	virtual bool CanRenameInternal() { return false; }
+
+	virtual void RenameInternal(FName NewName) {}
+	
+	/** Should return true if deletable. By default, we can delete items in the hierarchy, not in the source. */
+	virtual bool CanDeleteInternal() { return IsForHierarchy(); }
+
+	/** Should return true if the given drag drop operation is allowed to succeed. */
+	virtual FCanPerformActionResults CanDropOnInternal(TSharedPtr<FHierarchyElementViewModel>, EItemDropZone ItemDropZone);
+	
+	/** Override this to handle drop-on logic. This is called when an item has been dropped on the item that has implemented this function. */
+	virtual void OnDroppedOnInternal(TSharedPtr<FHierarchyElementViewModel>, EItemDropZone ItemDropZone) { }
+
+	/** Can be overridden to support sections in the source list.
+	 * In the hierarchy only categories can be parented directly to the root, but using this it is possible to add items to custom sections in the source panel.
+	 * This will only work for top-level objects, i.e. anything directly under the root. */
+	virtual const UHierarchySection* GetSectionInternal() const { return nullptr; }
+private:
+	/** Optionally implement this to refresh dependent data. */
+	virtual void RefreshChildrenDataInternal() {}
+	/** Optionally implement this to further customize the view model sync process.
+	 * An example for this is how the root view model handles sections, as sections exist outside the children hierarchy */
+	virtual void SyncViewModelsToDataInternal() {}
+	/** Optionally implement this to handle shutdown logic.
+	 * An example for this is when a section gets deleted, it iterates over all categories to null out the associated section  */
+	virtual void FinalizeInternal() {}
+
+	void PropagateOnChildSynced();
+protected:
+	/** The hierarchy element this view model represents. Assumed valid while this view model exists. */
+	UHierarchyElement* const Element;
+
+	/** Parent should be valid for all instances of this struct except for root objects */
+	TWeakPtr<FHierarchyElementViewModel> Parent;
+	TArray<TSharedPtr<FHierarchyElementViewModel>> Children;
+	
+	TWeakObjectPtr<UHierarchyViewModelBase> HierarchyViewModel;
+	
+	TArray<FOnFilterChild> ChildFilters;
+	mutable TArray<TSharedPtr<FHierarchyElementViewModel>> FilteredChildren;
+	
+	FSimpleDelegate OnRequestRenameDelegate;
+	FOnSynced OnSyncedDelegate;
+	FOnChildRequestedDeletion OnChildRequestedDeletionDelegate;
+	
+	bool bRenamePending = false;
+	bool bIsForHierarchy = false;
+};
+
+template <class DataClass, class ViewModelClass>
+void FHierarchyElementViewModel::GetChildrenViewModelsForType(TArray<TSharedPtr<ViewModelClass>>& OutChildren, bool bRecursive)
+{
+	for(auto& Child : Children)
+	{
+		if(Child->GetData()->IsA<DataClass>())
+		{
+			OutChildren.Add(StaticCastSharedPtr<ViewModelClass>(Child));
+		}
+	}
+
+	if(bRecursive)
+	{
+		for(auto& Child : Children)
+		{
+			Child->GetChildrenViewModelsForType<DataClass, ViewModelClass>(OutChildren, bRecursive);
+		}
+	}
+}
+
+struct HIERARCHYEDITOR_API FHierarchyRootViewModel : FHierarchyElementViewModel
+{
+	DECLARE_DELEGATE(FOnSyncPropagated)
+	DECLARE_DELEGATE(FOnSectionsChanged)
+	DECLARE_DELEGATE_OneParam(FOnSingleSectionChanged, TSharedPtr<FHierarchySectionViewModel> AddedSection)
+	
+	FHierarchyRootViewModel(UHierarchyElement* InItem, TWeakObjectPtr<UHierarchyViewModelBase> InHierarchyViewModel, bool bInIsForHierarchy) : FHierarchyElementViewModel(InItem, nullptr, InHierarchyViewModel, bInIsForHierarchy) {}
+	virtual ~FHierarchyRootViewModel() override;
+
+	virtual void Initialize() override;
+	
+	virtual FCanPerformActionResults CanDropOnInternal(TSharedPtr<FHierarchyElementViewModel> DraggedElement, EItemDropZone ItemDropZone) override;
+	virtual void OnDroppedOnInternal(TSharedPtr<FHierarchyElementViewModel> DroppedItem, EItemDropZone ItemDropZone) override;
+
+	virtual bool CanHaveChildren() const override { return true; }
+
+	TSharedPtr<struct FHierarchySectionViewModel> AddSection();
+	void DeleteSection(TSharedPtr<FHierarchyElementViewModel> SectionViewModel);
+	TArray<TSharedPtr<struct FHierarchySectionViewModel>>& GetSectionViewModels() { return SectionViewModels; }
+
+	FOnSyncPropagated& OnSyncPropagated() { return OnSyncPropagatedDelegate; }
+
+	/** General purpose delegate for when sections change */
+	FOnSectionsChanged& OnSectionsChanged() { return OnSectionsChangedDelegate; }
+	/** Delegates for when a section is added or removed */
+	FOnSingleSectionChanged& OnSectionAdded() { return OnSectionAddedDelegate; }
+	FOnSingleSectionChanged& OnSectionDeleted() { return OnSectionDeletedDelegate; }
+
+private:
+	void PropagateOnSynced();
+	virtual void SyncViewModelsToDataInternal() override;
+	TArray<TSharedPtr<struct FHierarchySectionViewModel>> SectionViewModels;
+
+	FOnSyncPropagated OnSyncPropagatedDelegate;
+	FOnSingleSectionChanged OnSectionAddedDelegate;
+	FOnSingleSectionChanged OnSectionDeletedDelegate;
+	FOnSectionsChanged OnSectionsChangedDelegate;
+};
+
+struct HIERARCHYEDITOR_API FHierarchySectionViewModel : FHierarchyElementViewModel
+{
+	FHierarchySectionViewModel(UHierarchySection* InItem, TSharedRef<FHierarchyRootViewModel> InParent, TWeakObjectPtr<UHierarchyViewModelBase> InHierarchyViewModel) : FHierarchyElementViewModel(InItem, InParent, InHierarchyViewModel, InParent->IsForHierarchy()) {}
+	
+	virtual ~FHierarchySectionViewModel() override {}
+
+	void SetSectionName(FName InSectionName) const { Cast<UHierarchySection>(Element)->SetSectionName(InSectionName); }
+	FName GetSectionName() const { return Cast<UHierarchySection>(Element)->GetSectionName(); }
+	
+	void SetSectionNameAsText(const FText& Text) const { Cast<UHierarchySection>(Element)->SetSectionNameAsText(Text); }
+	FText GetSectionNameAsText() const { return Cast<UHierarchySection>(Element)->GetSectionNameAsText(); }
+	FText GetSectionTooltip() const { return Cast<UHierarchySection>(Element)->GetTooltip(); }
+	
+	void SetSectionImage(const FSlateBrush* InSectionImage) { SectionImage = InSectionImage; }
+	const FSlateBrush* GetSectionImage() const { return SectionImage; }
+
+protected:
+	/** Only hierarchy sections are editable */
+	virtual FCanPerformActionResults IsEditableByUser() override { return FCanPerformActionResults(IsForHierarchy()); }
+	virtual bool CanHaveChildren() const override { return false; }
+	virtual FCanPerformActionResults CanDragInternal() override { return bIsForHierarchy; }
+	/** We can only rename hierarchy sections */
+	virtual bool CanRenameInternal() override { return IsForHierarchy(); }
+	virtual void RenameInternal(FName NewName) override { GetDataMutable<UHierarchySection>()->SetSectionName(NewName); }
+
+	virtual FCanPerformActionResults CanDropOnInternal(TSharedPtr<FHierarchyElementViewModel> DraggedItem, EItemDropZone ItemDropZone) override;
+	virtual void OnDroppedOnInternal(TSharedPtr<FHierarchyElementViewModel> DroppedItem, EItemDropZone ItemDropZone) override;
+
+	virtual void FinalizeInternal() override;
+
+private:
+	const FSlateBrush* SectionImage = nullptr;
+};
+
+struct HIERARCHYEDITOR_API FHierarchyItemViewModel : FHierarchyElementViewModel
+{
+	FHierarchyItemViewModel(UHierarchyItem* InElement, TSharedRef<FHierarchyElementViewModel> InParent, TWeakObjectPtr<UHierarchyViewModelBase> InHierarchyViewModel) : FHierarchyElementViewModel(InElement, InParent, InHierarchyViewModel, InParent->IsForHierarchy()) {}
+	
+	virtual ~FHierarchyItemViewModel() override {}
+
+	virtual FCanPerformActionResults IsEditableByUser() override { return FCanPerformActionResults(true); }
+	virtual bool CanHaveChildren() const override { return false; }
+	virtual FCanPerformActionResults CanDragInternal() override { return true; }
+
+	virtual FCanPerformActionResults CanDropOnInternal(TSharedPtr<FHierarchyElementViewModel> DraggedItem, EItemDropZone ItemDropZone) override;
+	virtual void OnDroppedOnInternal(TSharedPtr<FHierarchyElementViewModel> DroppedItem, EItemDropZone ItemDropZone) override;
+};
+
+struct HIERARCHYEDITOR_API FHierarchyCategoryViewModel : FHierarchyElementViewModel
+{
+	FHierarchyCategoryViewModel(UHierarchyCategory* InCategory, TSharedRef<FHierarchyElementViewModel> InParent, TWeakObjectPtr<UHierarchyViewModelBase> InHierarchyViewModel) : FHierarchyElementViewModel(InCategory, InParent, InHierarchyViewModel, InParent->IsForHierarchy()) {}
+	virtual ~FHierarchyCategoryViewModel() override{}
+
+	FText GetCategoryName() const { return GetData<UHierarchyCategory>()->GetCategoryAsText(); }
+	
+	virtual FCanPerformActionResults IsEditableByUser() override { return FCanPerformActionResults(true); }
+	virtual bool CanHaveChildren() const override { return true; }
+	virtual FCanPerformActionResults CanDragInternal() override { return true; }
+	virtual bool CanRenameInternal() override { return true; }
+	virtual void RenameInternal(FName NewName) override { GetDataMutable<UHierarchyCategory>()->SetCategoryName(NewName); }
+	virtual const UHierarchySection* GetSectionInternal() const override { return Cast<UHierarchyCategory>(Element)->GetSection(); }
+
+	bool IsTopCategoryActive() const;
+
+	virtual FCanPerformActionResults CanDropOnInternal(TSharedPtr<FHierarchyElementViewModel>, EItemDropZone ItemDropZone) override;
+	virtual void OnDroppedOnInternal(TSharedPtr<FHierarchyElementViewModel> DroppedItem, EItemDropZone ItemDropZone) override;
+};
+
+class HIERARCHYEDITOR_API FSectionDragDropOp : public FHierarchyDragDropOp
+{
+public:
+	DRAG_DROP_OPERATOR_TYPE(FSectionDragDropOp, FHierarchyDragDropOp)
+
+	FSectionDragDropOp(TSharedPtr<FHierarchySectionViewModel> SectionViewModel) : FHierarchyDragDropOp(SectionViewModel) {}
+	
+	TWeakPtr<FHierarchySectionViewModel> GetDraggedSection() const { return StaticCastSharedPtr<FHierarchySectionViewModel>(DraggedElement.Pin()); }
+private:
+	virtual TSharedRef<SWidget> CreateCustomDecorator() const override;
+};
diff --git a/Engine/Source/Editor/HierarchyEditor/Public/Widgets/SHierarchyEditor.h b/Engine/Source/Editor/HierarchyEditor/Public/Widgets/SHierarchyEditor.h
new file mode 100644
--- /dev/null
+++ b/Engine/Source/Editor/HierarchyEditor/Public/Widgets/SHierarchyEditor.h
@@ -0,0 +1,284 @@
+// Copyright Epic Games, Inc. All Rights Reserved.
+
+#pragma once
+
+#include "CoreMinimal.h"
+#include "SDropTarget.h"
+#include "IDetailsView.h"
+#include "HierarchyEditorStyle.h"
+#include "Widgets/Views/STreeView.h"
+#include "Widgets/Views/STableRow.h"
+#include "Widgets/Text/SInlineEditableTextBlock.h"
+#include "Widgets/Input/SMenuAnchor.h"
+#include "Styling/SlateTypes.h"
+#include "HierarchyViewModelBase.h"
+#include "Widgets/Input/SSearchBox.h"
+#include "Widgets/Input/SCheckBox.h"
+#include "Misc/NotifyHook.h"
+#include "UObject/GCObject.h"
+
+class UHierarchyViewModelBase;
+class UHierarchySection;
+struct FHierarchyElementViewModel;
+struct FHierarchySectionViewModel;
+
+class HIERARCHYEDITOR_API SHierarchySection : public SCompoundWidget
+{
+	DECLARE_DELEGATE_OneParam(FOnSectionActivated, TSharedPtr<FHierarchySectionViewModel> SectionViewModel)
+	
+	SLATE_BEGIN_ARGS(SHierarchySection)
+		: _bForbidDropOn(false)
+		{}
+		SLATE_ATTRIBUTE(ECheckBoxState, IsSectionActive)
+		SLATE_EVENT(FOnSectionActivated, OnSectionActivated)
+		SLATE_ARGUMENT(bool, bForbidDropOn)
+	SLATE_END_ARGS()
+
+public:
+	void Construct(const FArguments& InArgs, TSharedPtr<struct FHierarchySectionViewModel> InSection, TWeakObjectPtr<UHierarchyViewModelBase> HierarchyViewModel);
+	virtual ~SHierarchySection() override;
+
+	bool OnCanAcceptDrop(TSharedPtr<FDragDropOperation> DragDropOperation, EItemDropZone ItemDropZone) const;
+	FReply OnDroppedOn(const FGeometry&, const FDragDropEvent& DragDropEvent, EItemDropZone DropZone) const;
+
+	void TryEnterEditingMode() const;
+
+	TSharedPtr<struct FHierarchySectionViewModel> GetSectionViewModel();
+private:
+	virtual FReply OnMouseButtonDown(const FGeometry& MyGeometry, const FPointerEvent& MouseEvent) override;
+	virtual FReply OnMouseButtonUp(const FGeometry& MyGeometry, const FPointerEvent& MouseEvent) override;
+	
+	virtual void OnDragEnter(const FGeometry& MyGeometry, const FDragDropEvent& DragDropEvent) override;
+	virtual FReply OnDragDetected(const FGeometry& MyGeometry, const FPointerEvent& MouseEvent) override;
+	virtual void OnDragLeave(const FDragDropEvent& DragDropEvent) override;
+
+	virtual FReply OnKeyDown(const FGeometry& MyGeometry, const FKeyEvent& InKeyEvent) override;
+	
+	virtual void OnMouseLeave(const FPointerEvent& MouseEvent) override;
+
+	UHierarchySection* TryGetSectionData() const;
+
+	FText GetText() const;
+	FText GetTooltipText() const;
+	void OnRenameSection(const FText& Text, ETextCommit::Type CommitType) const;
+	bool OnVerifySectionRename(const FText& NewName, FText& OutTooltip) const;
+
+	bool IsSectionSelected() const;
+	bool IsSectionReadOnly() const;
+	ECheckBoxState GetSectionCheckState() const;
+	void OnSectionCheckChanged(ECheckBoxState NewState);
+	EActiveTimerReturnType ActivateSectionIfDragging(double CurrentTime, float DeltaTime) const;
+private:
+	TSharedPtr<SMenuAnchor> MenuAnchor;
+	TSharedPtr<SCheckBox> CheckBox;
+	TSharedPtr<SInlineEditableTextBlock> InlineEditableTextBlock;
+	TWeakObjectPtr<UHierarchyViewModelBase> HierarchyViewModel;
+	TSharedPtr<FHierarchySectionViewModel> SectionViewModel;
+private:
+	TAttribute<ECheckBoxState> IsSectionActive;
+	FOnSectionActivated OnSectionActivatedDelegate;
+	bool bForbidDropOn = false;
+	mutable bool bDraggedOn = false;
+};
+
+class HIERARCHYEDITOR_API SHierarchyEditor : public SCompoundWidget, public FGCObject, public FNotifyHook
+{
+	DECLARE_DELEGATE_RetVal_OneParam(TSharedRef<SWidget>, FOnGenerateRowContentWidget, TSharedRef<FHierarchyElementViewModel> HierarchyElement);
+	DECLARE_DELEGATE_RetVal_OneParam(TSharedRef<SWidget>, FOnGenerateCustomDetailsPanelNameWidget, TSharedPtr<FHierarchyElementViewModel> HierarchyElement);
+
+	SLATE_BEGIN_ARGS(SHierarchyEditor)
+		: _bReadOnly(true)
+		, _ItemRowStyle(&FAppStyle::Get().GetWidgetStyle<FTableRowStyle>("TableView.Row"))
+		, _CategoryRowStyle(&FHierarchyEditorStyle::Get().GetWidgetStyle<FTableRowStyle>("HierarchyEditor.Row.Category"))
+	{}
+		SLATE_ARGUMENT(bool, bReadOnly)
+		SLATE_EVENT(FOnGenerateRowContentWidget, OnGenerateRowContentWidget)
+		SLATE_EVENT(FOnGenerateCustomDetailsPanelNameWidget, OnGenerateCustomDetailsPanelNameWidget)
+		SLATE_STYLE_ARGUMENT(FTableRowStyle, ItemRowStyle)
+		SLATE_STYLE_ARGUMENT(FTableRowStyle, CategoryRowStyle)
+	SLATE_END_ARGS()
+
+public:
+	void Construct(const FArguments& InArgs, TObjectPtr<UHierarchyViewModelBase> InHierarchyViewModel);
+	virtual ~SHierarchyEditor() override;
+
+	void RefreshSourceItems();
+	void RefreshAllViews(bool bFullRefresh = false);
+	void RequestRefreshAllViewsNextFrame(bool bFullRefresh = false);
+	void RefreshSourceView(bool bFullRefresh = false) const;
+	void RequestRefreshSourceViewNextFrame(bool bFullRefresh = false);
+	void RefreshHierarchyView(bool bFullRefresh = false) const;
+	void RequestRefreshHierarchyViewNextFrame(bool bFullRefresh = false);
+	void RefreshSectionsView();
+	void RequestRefreshSectionsViewNextFrame();
+
+	void NavigateToHierarchyElement(FHierarchyElementIdentity Identity) const;
+	void NavigateToHierarchyElement(TSharedPtr<FHierarchyElementViewModel> Item) const;
+	bool IsItemSelected(TSharedPtr<FHierarchyElementViewModel> Item) const;
+	
+private:
+	// need to do this to enable focus so we can handle shortcuts
+	virtual bool SupportsKeyboardFocus() const override { return true; }
+	
+	TSharedRef<ITableRow> GenerateSourceItemRow(TSharedPtr<FHierarchyElementViewModel> HierarchyItem, const TSharedRef<STableViewBase>& TableViewBase);
+	TSharedRef<ITableRow> GenerateHierarchyItemRow(TSharedPtr<FHierarchyElementViewModel> HierarchyItem, const TSharedRef<STableViewBase>& TableViewBase);
+
+	bool FilterForSourceSection(TSharedPtr<const FHierarchyElementViewModel> ItemViewModel) const;
+private:
+	void Reinitialize();
+
+	void BindToHierarchyRootViewModel();
+	void UnbindFromHierarchyRootViewModel() const;
+	
+	/** Source items reflect the base, unedited status of items to edit into a hierarchy */
+	const TArray<TSharedPtr<FHierarchyElementViewModel>>& GetSourceItems() const;
+	
+	bool IsDetailsPanelEditingAllowed() const;
+	
+	void ClearSourceItems() const;
+	
+	void RequestRenameSelectedItem();
+	bool CanRequestRenameSelectedItem() const;
+
+	void DeleteItems(TArray<TSharedPtr<FHierarchyElementViewModel>> ItemsToDelete) const;
+	void DeleteSelectedHierarchyItems() const;
+	bool CanDeleteSelectedElements() const;
+
+	void NavigateToMatchingHierarchyElementFromSelectedSourceElement() const;
+	bool CanNavigateToMatchingHierarchyElementFromSelectedSourceElement() const;
+
+	void DeleteActiveSection() const;
+	bool CanDeleteActiveSection() const;
+
+	void OnElementAdded(TSharedPtr<FHierarchyElementViewModel> AddedItem);
+	void OnHierarchySectionActivated(TSharedPtr<FHierarchySectionViewModel> Section);
+	void OnSourceSectionActivated(TSharedPtr<FHierarchySectionViewModel> Section);
+	void OnHierarchySectionAdded(TSharedPtr<FHierarchySectionViewModel> AddedSection);
+	void OnHierarchySectionDeleted(TSharedPtr<FHierarchySectionViewModel> DeletedSection);
+
+	void SetActiveSourceSection(TSharedPtr<struct FHierarchySectionViewModel>);
+	TSharedPtr<FHierarchySectionViewModel> GetActiveSourceSection() const;
+	UHierarchySection* GetActiveSourceSectionData() const;
+	
+	void OnSelectionChanged(TSharedPtr<FHierarchyElementViewModel> HierarchyItem, ESelectInfo::Type Type, bool bFromHierarchy) const;
+
+	void RunSourceSearch();
+	void OnSourceSearchTextChanged(const FText& Text);
+	void OnSourceSearchTextCommitted(const FText& Text, ETextCommit::Type CommitType);
+	void OnSearchButtonClicked(SSearchBox::SearchDirection SearchDirection);
+
+	virtual FReply OnKeyDown(const FGeometry& MyGeometry, const FKeyEvent& InKeyEvent) override;
+	virtual FReply OnMouseButtonDown(const FGeometry& MyGeometry, const FPointerEvent& MouseEvent) override;
+	virtual FReply OnMouseButtonUp(const FGeometry& MyGeometry, const FPointerEvent& MouseEvent) override;
+
+	FReply OnAddCategoryClicked() const;
+	FReply OnAddSectionClicked() const;
+
+	TSharedPtr<SWidget> SummonContextMenuForSelectedRows(bool bFromHierarchy) const;
+
+	struct FSearchItem
+	{
+		TArray<TSharedPtr<FHierarchyElementViewModel>> Path;
+
+		TSharedPtr<FHierarchyElementViewModel> GetEntry() const
+		{
+			return Path.Num() > 0 ? 
+				Path[Path.Num() - 1] : 
+				nullptr;
+		}
+
+		bool operator==(const FSearchItem& Item) const
+		{
+			return Path == Item.Path;
+		}
+	};
+
+	/** This will recursively generated parent chain paths for all items within the given root. Used for expansion purposes. */
+	void GenerateSearchItems(TSharedRef<FHierarchyElementViewModel> Root, TArray<TSharedPtr<FHierarchyElementViewModel>> ParentChain, TArray<FSearchItem>& OutSearchItems);
+	void ExpandSourceSearchResults();
+	void SelectNextSourceSearchResult();
+	void SelectPreviousSourceSearchResult();
+	TOptional<SSearchBox::FSearchResultData> GetSearchResultData() const;
+	
+	FHierarchyElementViewModel::FCanPerformActionResults CanDropOnRoot(TSharedPtr<FHierarchyElementViewModel> DraggedItem) const;
+
+	/** Callback functions for the root widget */
+	FReply HandleHierarchyRootDrop(const FGeometry&, const FDragDropEvent& DragDropEvent) const;
+	bool OnCanDropOnRoot(TSharedPtr<FDragDropOperation> DragDropOperation) const;
+	void OnRootDragEnter(const FDragDropEvent& DragDropEvent) const;
+	void OnRootDragLeave(const FDragDropEvent& DragDropEvent) const;
+	FSlateColor GetRootIconColor() const;
+
+	virtual FString GetReferencerName() const override; 
+	virtual void AddReferencedObjects(FReferenceCollector& Collector) override;
+
+	virtual void NotifyPostChange(const FPropertyChangedEvent& PropertyChangedEvent, FProperty* PropertyThatChanged) override;
+private:
+	TWeakObjectPtr<UHierarchyViewModelBase> HierarchyViewModel;
+	
+	TArray<FSearchItem> SourceSearchResults;
+	TOptional<FSearchItem> FocusedSearchResult;
+
+	const FTableRowStyle* CategoryRowStyle = nullptr;
+	const FTableRowStyle* ItemRowStyle = nullptr;
+	
+	mutable TWeakPtr<FHierarchyElementViewModel> SelectedDetailsPanelItemViewModel;
+
+private:
+	TObjectPtr<UHierarchyRoot> SourceRoot = nullptr;
+	TSharedPtr<FHierarchyRootViewModel> SourceRootViewModel;
+	TSharedPtr<STreeView<TSharedPtr<FHierarchyElementViewModel>>> SourceTreeView;
+	TSharedPtr<STreeView<TSharedPtr<FHierarchyElementViewModel>>> HierarchyTreeView;
+	TWeakPtr<struct FHierarchySectionViewModel> ActiveSourceSection;
+	TSharedPtr<class SWrapBox> SourceSectionBox;
+	TSharedPtr<class SWrapBox> HierarchySectionBox;
+	TSharedPtr<SSearchBox> SourceSearchBox;
+	TSharedPtr<IDetailsView> DetailsPanel;
+private:
+	FOnGenerateRowContentWidget OnGenerateRowContentWidget;
+	FOnGenerateCustomDetailsPanelNameWidget OnGenerateCustomDetailsPanelNameWidget;
+	TSharedPtr<FActiveTimerHandle> RefreshHierarchyViewNextFrameHandle;
+	TSharedPtr<FActiveTimerHandle> RefreshSourceViewNextFrameHandle;
+	TSharedPtr<FActiveTimerHandle> RefreshSectionsViewNextFrameHandle;
+};
+
+class HIERARCHYEDITOR_API SHierarchyCategory : public SCompoundWidget
+{
+	SLATE_BEGIN_ARGS(SHierarchyCategory)
+	{}
+		SLATE_EVENT(FIsSelected, IsSelected)
+	SLATE_END_ARGS()
+
+public:
+	void Construct(const FArguments& InArgs, TSharedPtr<struct FHierarchyCategoryViewModel> InCategoryViewModel);
+
+	void EnterEditingMode() const;
+	bool OnVerifyCategoryRename(const FText& NewName, FText& OutTooltip) const;
+	
+	FText GetCategoryText() const;
+	void OnRenameCategory(const FText& NewText, ETextCommit::Type) const;
+	
+private:
+	TWeakPtr<struct FHierarchyCategoryViewModel> CategoryViewModel;
+	TSharedPtr<SInlineEditableTextBlock> InlineEditableTextBlock;
+};
+
+class SSectionDragDropTarget : public SDropTarget
+{
+public:
+	SLATE_BEGIN_ARGS(SSectionDragDropTarget)
+		: _DropTargetArgs(SDropTarget::FArguments())
+	{}	
+		SLATE_ARGUMENT(SDropTarget::FArguments, DropTargetArgs)
+	SLATE_END_ARGS()
+
+	void Construct(const FArguments& InArgs, TSharedPtr<FHierarchySectionViewModel> OwningSection, EItemDropZone ItemDropZone);
+	
+	virtual FReply OnDragOver(const FGeometry& MyGeometry, const FDragDropEvent& DragDropEvent) override;
+	virtual void OnDragLeave(const FDragDropEvent& DragDropEvent) override;
+
+private:
+	TSharedPtr<FHierarchySectionViewModel> OwningSection;
+	/** The drop zone this drop target represents */
+	EItemDropZone DropZone = EItemDropZone::OntoItem;
+};
diff --git a/Engine/Source/Editor/UnrealEd/Private/EditorEngine.cpp b/Engine/Source/Editor/UnrealEd/Private/EditorEngine.cpp
--- a/Engine/Source/Editor/UnrealEd/Private/EditorEngine.cpp
+++ b/Engine/Source/Editor/UnrealEd/Private/EditorEngine.cpp
@@ -1512,7 +1512,8 @@ void UEditorEngine::LoadDefaultEditorModules()
 			TEXT("StructUtilsEditor"),
 			TEXT("StructUtilsTestSuite"),
 			TEXT("MassEntityEditor"),
-			TEXT("SVGDistanceField")
+			TEXT("SVGDistanceField"),
+			TEXT("HierarchyEditor")
 		};
 
 	FScopedSlowTask ModuleSlowTask((float)UE_ARRAY_COUNT(ModuleNames));
