Fixes to HierarchyEditor context menus
- Each HierarchyViewModel now generates a unique menu. Previously all shared the same menu name.
- Fix for all menu generation being bound to a specific HierarchyViewModel. Now, the function is static and will call an Internal function that allows overriding menu entries
- Fixed 'Find in Hierarchy' showing up for hierarchy elements in the hierarchy (should only be available in the source)

diff --git a/Engine/Source/Editor/HierarchyEditor/Private/HierarchyViewModelBase.cpp b/Engine/Source/Editor/HierarchyEditor/Private/HierarchyViewModelBase.cpp
--- a/Engine/Source/Editor/HierarchyEditor/Private/HierarchyViewModelBase.cpp
+++ b/Engine/Source/Editor/HierarchyEditor/Private/HierarchyViewModelBase.cpp
@@ -411,11 +411,11 @@ void UHierarchyViewModelBase::Initialize()
 
 	UToolMenus* ToolMenus = UToolMenus::Get();
 
-	FName MenuName("HierarchyMenu");
+	FName MenuName = GetContextMenuName();
 	if(!ToolMenus->IsMenuRegistered(MenuName))
 	{
 		UToolMenu* HierarchyMenu = ToolMenus->RegisterMenu(MenuName, NAME_None, EMultiBoxType::Menu);
-		HierarchyMenu->AddDynamicSection(NAME_None, FNewToolMenuDelegate::CreateUObject(this, &UHierarchyViewModelBase::GenerateDynamicContextMenu));
+		HierarchyMenu->AddDynamicSection(NAME_None, FNewToolMenuDelegate::CreateStatic(&UHierarchyViewModelBase::GenerateDynamicContextMenu));
 	}
 	
 	SetupCommands();
@@ -669,6 +669,11 @@ const TArray<TSharedPtr<FHierarchyElementViewModel>>& UHierarchyViewModelBase::G
 	return HierarchyRootViewModel->GetFilteredChildren();
 }
 
+FName UHierarchyViewModelBase::GetContextMenuName() const
+{
+	return FName(FString(TEXT("HierarchyEditor.") + GetClass()->GetName()));
+}
+
 TSharedRef<FHierarchyDragDropOp> UHierarchyViewModelBase::CreateDragDropOp(TSharedRef<FHierarchyElementViewModel> Item)
 {
 	TSharedRef<FHierarchyDragDropOp> DragDropOp = MakeShared<FHierarchyDragDropOp>(Item);
@@ -1644,25 +1649,40 @@ void UHierarchyViewModelBase::AddSection() const
 	OnHierarchyChangedDelegate.Broadcast();
 }
 
-void UHierarchyViewModelBase::GenerateDynamicContextMenu(UToolMenu* ToolMenu) const
+void UHierarchyViewModelBase::GenerateDynamicContextMenu(UToolMenu* ToolMenu)
 {
 	UHierarchyMenuContext* HierarchyMenuContext = ToolMenu->FindContext<UHierarchyMenuContext>();
 
-	if(HierarchyMenuContext == nullptr)
+	if(HierarchyMenuContext == nullptr || HierarchyMenuContext->HierarchyViewModel.IsValid() == false)
 	{
 		return;
 	}
-	
-	ToolMenu->AddMenuEntry("Dynamic", FToolMenuEntry::InitMenuEntryWithCommandList(FHierarchyEditorCommands::Get().FindInHierarchy, Commands, TAttribute<FText>(), TAttribute<FText>(), FSlateIcon(FAppStyle::GetAppStyleSetName(), "Icons.Find")));
-	ToolMenu->AddMenuEntry("Dynamic", FToolMenuEntry::InitMenuEntryWithCommandList(FGenericCommands::Get().Rename, Commands));
-	ToolMenu->AddMenuEntry("Dynamic", FToolMenuEntry::InitMenuEntryWithCommandList(FGenericCommands::Get().Delete, Commands));
 
+	UHierarchyViewModelBase* HierarchyViewModel = HierarchyMenuContext->HierarchyViewModel.Get();
+	HierarchyViewModel->GenerateDynamicContextMenuInternal(ToolMenu);
+	
 	if(HierarchyMenuContext->MenuHierarchyElements.Num() == 1)
 	{
 		HierarchyMenuContext->MenuHierarchyElements[0]->AppendDynamicContextMenuForSingleElement(ToolMenu);
 	}
 }
 
+void UHierarchyViewModelBase::GenerateDynamicContextMenuInternal(UToolMenu* DynamicToolMenu) const
+{
+	UHierarchyMenuContext* HierarchyMenuContext = DynamicToolMenu->FindContext<UHierarchyMenuContext>();
+
+	if(HierarchyMenuContext == nullptr || HierarchyMenuContext->HierarchyViewModel.IsValid() == false)
+	{
+		return;
+	}
+
+	UHierarchyViewModelBase* HierarchyViewModel = HierarchyMenuContext->HierarchyViewModel.Get();
+
+	DynamicToolMenu->AddMenuEntry("Dynamic", FToolMenuEntry::InitMenuEntryWithCommandList(FHierarchyEditorCommands::Get().FindInHierarchy, HierarchyViewModel->GetCommands(), TAttribute<FText>(), TAttribute<FText>(), FSlateIcon(FAppStyle::GetAppStyleSetName(), "Icons.Find")));
+	DynamicToolMenu->AddMenuEntry("Dynamic", FToolMenuEntry::InitMenuEntryWithCommandList(FGenericCommands::Get().Rename, HierarchyViewModel->GetCommands()));
+	DynamicToolMenu->AddMenuEntry("Dynamic", FToolMenuEntry::InitMenuEntryWithCommandList(FGenericCommands::Get().Delete, HierarchyViewModel->GetCommands()));
+}
+
 UHierarchyElement* UHierarchyViewModelBase::AddElementUnderRoot(TSubclassOf<UHierarchyElement> NewChildClass, FHierarchyElementIdentity ChildIdentity)
 {
 	FScopedTransaction Transaction(LOCTEXT("Transaction_AddItem", "Add hierarchy item"));
diff --git a/Engine/Source/Editor/HierarchyEditor/Private/Widgets/SHierarchyEditor.cpp b/Engine/Source/Editor/HierarchyEditor/Private/Widgets/SHierarchyEditor.cpp
--- a/Engine/Source/Editor/HierarchyEditor/Private/Widgets/SHierarchyEditor.cpp
+++ b/Engine/Source/Editor/HierarchyEditor/Private/Widgets/SHierarchyEditor.cpp
@@ -30,12 +30,13 @@ TSharedRef<SWidget> SummonContextMenu(TArray<TSharedPtr<FHierarchyElementViewMod
 {
 	UHierarchyMenuContext* MenuContextObject = NewObject<UHierarchyMenuContext>();
 	MenuContextObject->MenuHierarchyElements = MenuHierarchyElements;
+	TWeakObjectPtr<UHierarchyViewModelBase> ViewModel = MenuHierarchyElements[0]->GetHierarchyViewModel();
+	MenuContextObject->HierarchyViewModel = ViewModel;
 
 	FToolMenuContext MenuContext(MenuContextObject);
-	TWeakObjectPtr<UHierarchyViewModelBase> ViewModel = MenuHierarchyElements[0]->GetHierarchyViewModel();
 	MenuContext.AppendCommandList(ViewModel->GetCommands());
 	
-	TSharedRef<SWidget> MenuWidget = UToolMenus::Get()->GenerateWidget("HierarchyMenu", MenuContext);
+	TSharedRef<SWidget> MenuWidget = UToolMenus::Get()->GenerateWidget(ViewModel->GetContextMenuName(), MenuContext);
 	return MenuWidget;
 }
 
@@ -751,7 +752,8 @@ void SHierarchyEditor::Construct(const FArguments& InArgs, TObjectPtr<UHierarchy
 	
 	HierarchyViewModel->GetCommands()->MapAction(FHierarchyEditorCommands::Get().FindInHierarchy,
 		FExecuteAction::CreateSP(this, &SHierarchyEditor::NavigateToMatchingHierarchyElementFromSelectedSourceElement),
-		FCanExecuteAction::CreateSP(this, &SHierarchyEditor::CanNavigateToMatchingHierarchyElementFromSelectedSourceElement));
+		FCanExecuteAction::CreateSP(this, &SHierarchyEditor::CanNavigateToMatchingHierarchyElementFromSelectedSourceElement),
+		FIsActionChecked(), FIsActionButtonVisible::CreateSP(this, &SHierarchyEditor::CanNavigateToMatchingHierarchyElementFromSelectedSourceElement));
 	
 	HierarchyViewModel->ForceFullRefresh();
 }
@@ -816,11 +818,12 @@ void SHierarchyEditor::RequestRefreshAllViewsNextFrame(bool bFullRefresh)
 FReply SHierarchyEditor::OnKeyDown(const FGeometry& MyGeometry, const FKeyEvent& InKeyEvent)
 {
 	UHierarchyMenuContext* HierarchyMenuContext = NewObject<UHierarchyMenuContext>();
-	 
+	HierarchyMenuContext->HierarchyViewModel = HierarchyViewModel;
+	
 	FToolMenuContext Context;
 	Context.AddObject(HierarchyMenuContext);
 	
-	UToolMenus::Get()->GenerateMenu("HierarchyMenu", Context);
+	UToolMenus::Get()->GenerateMenu(HierarchyViewModel->GetContextMenuName(), Context);
 	return HierarchyViewModel->GetCommands()->ProcessCommandBindings(InKeyEvent) ? FReply::Handled() : FReply::Unhandled();
 }
 
diff --git a/Engine/Source/Editor/HierarchyEditor/Public/HierarchyViewModelBase.h b/Engine/Source/Editor/HierarchyEditor/Public/HierarchyViewModelBase.h
--- a/Engine/Source/Editor/HierarchyEditor/Public/HierarchyViewModelBase.h
+++ b/Engine/Source/Editor/HierarchyEditor/Public/HierarchyViewModelBase.h
@@ -419,6 +419,7 @@ class HIERARCHYEDITOR_API UHierarchyMenuContext : public UObject
 	GENERATED_BODY()
 
 public:
+	TWeakObjectPtr<class UHierarchyViewModelBase> HierarchyViewModel;
 	TArray<TSharedPtr<FHierarchyElementViewModel>> MenuHierarchyElements;
 };
 
@@ -450,6 +451,8 @@ public:
 	bool IsFinalized() const { return bIsFinalized; }
 	bool IsValid() const { return IsInitialized() && !IsFinalized(); }
 
+	FName GetContextMenuName () const;
+
 	/** Creates view model hierarchy elements. To create custom view models, override CreateCustomViewModelForElement. */
 	TSharedPtr<FHierarchyElementViewModel> CreateViewModelForElement(UHierarchyElement* Element, TSharedPtr<FHierarchyElementViewModel> Parent);
 	/** Get the root view model associated with the hierarchy. */
@@ -486,11 +489,10 @@ private:
 
 	/** This function is used to determine custom view models for Hierarchy Elements. Called by CreateViewModelForElement. */
 	virtual TSharedPtr<FHierarchyElementViewModel> CreateCustomViewModelForElement(UHierarchyElement* Element, TSharedPtr<FHierarchyElementViewModel> Parent) { return nullptr; }
-	/** Override this to add custom context menu options. The items to generate this for are contained in a UHierarchyMenuContext object accessible through the ToolMenu. */
-	virtual void GenerateDynamicContextMenuInternal(TArray<TSharedPtr<FHierarchyElementViewModel>> ViewModels) const {}
+	/** Override this to add custom context menu options. Defaults to common actions such as finding, renaming and deleting.
+	 * The hierarchy elements to generate this menu for are contained in a UHierarchyMenuContext object accessible through the ToolMenu. */
+	virtual void GenerateDynamicContextMenuInternal(UToolMenu* DynamicToolMenu) const;
 public:
-	virtual void GenerateDynamicContextMenu(UToolMenu* ToolMenu) const;
-
 	const UHierarchyDataRefreshContext* GetRefreshContext() const { return RefreshContext; }
 	/** Set the refresh context to easily allow Hierarchy Elements representing external data to access whether the external data still exists. */
 	void SetRefreshContext(UHierarchyDataRefreshContext* InContext) { RefreshContext = InContext; }
@@ -552,6 +554,8 @@ private:
 	bool FilterForHierarchySection(TSharedPtr<const FHierarchyElementViewModel> ViewModel) const;
 	bool FilterForUncategorizedRootItemsInAllSection(TSharedPtr<const FHierarchyElementViewModel> ViewModel) const;
 
+	static void GenerateDynamicContextMenu(UToolMenu* ToolMenu);
+
 	void ToolMenuRequestRename(const FToolMenuContext& Context) const;
 	bool ToolMenuCanRequestRename(const FToolMenuContext& Context) const;
 
