Fixed crash in HierarchyEditor due to unsaved hierarchy elements due to RF_Transient flag.
Also added some more logging & safety mechanisms to help with these kinds of issues in the future.

diff --git a/Engine/Plugins/FX/Niagara/Source/NiagaraEditor/Private/ViewModels/HierarchyEditor/NiagaraHierarchyScriptParametersViewModel.cpp b/Engine/Plugins/FX/Niagara/Source/NiagaraEditor/Private/ViewModels/HierarchyEditor/NiagaraHierarchyScriptParametersViewModel.cpp
--- a/Engine/Plugins/FX/Niagara/Source/NiagaraEditor/Private/ViewModels/HierarchyEditor/NiagaraHierarchyScriptParametersViewModel.cpp
+++ b/Engine/Plugins/FX/Niagara/Source/NiagaraEditor/Private/ViewModels/HierarchyEditor/NiagaraHierarchyScriptParametersViewModel.cpp
@@ -7,12 +7,18 @@
 #include "NiagaraGraph.h"
 #include "NiagaraScriptSource.h"
 #include "ScopedTransaction.h"
+#include "Logging/StructuredLog.h"
 #include "ViewModels/NiagaraScriptViewModel.h"
 
 #define LOCTEXT_NAMESPACE "NiagaraScriptParameterHierarchyEditor"
 
 void UNiagaraHierarchyScriptParameter::Initialize(const UNiagaraScriptVariable& InParameterScriptVariable)
 {
+	if(InParameterScriptVariable.Metadata.GetVariableGuid().IsValid() == false)
+	{
+		UE_LOGFMT(LogNiagaraEditor, Warning, "Invalid hierarchy script parameter initialization. No valid guid found. This hierarchy element will be deleted on next refresh.");
+	}
+	
 	SetIdentity(FHierarchyElementIdentity({InParameterScriptVariable.Metadata.GetVariableGuid()}, {}));
 }
 
diff --git a/Engine/Source/Editor/HierarchyEditor/Private/HierarchyViewModelBase.cpp b/Engine/Source/Editor/HierarchyEditor/Private/HierarchyViewModelBase.cpp
--- a/Engine/Source/Editor/HierarchyEditor/Private/HierarchyViewModelBase.cpp
+++ b/Engine/Source/Editor/HierarchyEditor/Private/HierarchyViewModelBase.cpp
@@ -17,6 +17,7 @@
 #include "ScopedTransaction.h"
 #include "Framework/Notifications/NotificationManager.h"
 #include "ToolMenus.h"
+#include "Logging/StructuredLog.h"
 #include "Widgets/Notifications/SNotificationList.h"
 
 #define LOCTEXT_NAMESPACE "HierarchyEditor"
@@ -142,6 +143,22 @@ void UHierarchyElement::PostLoad()
 		SetIdentity(FHierarchyElementIdentity({Guid_DEPRECATED}, {}));
 	}
 
+	bool bAnyChildNullptr = false;		
+	for(auto It = Children.CreateIterator(); It; ++It)
+	{
+		if(*It == nullptr)
+		{
+			bAnyChildNullptr = true;
+			It.RemoveCurrent();
+		}
+	}
+
+	if(bAnyChildNullptr)
+	{
+		UPackage* Package = GetPackage();
+		UE_LOG(LogHierarchyEditor, Warning, TEXT("HierarchyElement %s found nullptr child in asset %s. Removed all nullptr children. This is indicative of something wrong. Check if the hierarchy is still correct and fix it, if necessary."), *ToString(), *GetNameSafe(Package))	
+	}
+	
 	Super::PostLoad();
 }
 
@@ -252,6 +269,25 @@ void UHierarchyRoot::EmptyAllData()
 	Sections.Empty();
 }
 
+void UHierarchyRoot::Serialize(FStructuredArchive::FRecord Record)
+{
+	// If the root isn't transient, neither should any of its hierarchy elements be.
+	// This is expected to happen as the source elements are transient by default.
+	// When source hierarchy elements are put into the hierarchy we have to make sure to remove the flag after
+	if(Record.GetArchiveState().IsSaving() && this->HasAnyFlags(RF_Transient) == false)
+	{
+		TArray<UHierarchyElement*> AllElements;
+		GetChildrenOfType(AllElements, true);
+
+		for(UHierarchyElement* Element : AllElements)
+		{
+			Element->ClearFlags(RF_Transient);
+		}
+	}
+	
+	Super::Serialize(Record);
+}
+
 bool FHierarchyCategoryViewModel::IsTopCategoryActive() const
 {
 	if(UHierarchyCategory* Category = GetDataMutable<UHierarchyCategory>())
@@ -752,6 +788,7 @@ void FHierarchyElementViewModel::RefreshChildrenData()
 	{
 		if(Child->RepresentsExternalData() && Child->DoesExternalDataStillExist(HierarchyViewModel->GetRefreshContext()) == false)
 		{
+			UE_LOGFMT(LogHierarchyEditor, Verbose, "Hierarchy Element {ElementName} no longer has valid external data. Deleting.", Child->ToString());
 			Child->Delete();
 		}
 	}
diff --git a/Engine/Source/Editor/HierarchyEditor/Public/HierarchyViewModelBase.h b/Engine/Source/Editor/HierarchyEditor/Public/HierarchyViewModelBase.h
--- a/Engine/Source/Editor/HierarchyEditor/Public/HierarchyViewModelBase.h
+++ b/Engine/Source/Editor/HierarchyEditor/Public/HierarchyViewModelBase.h
@@ -319,6 +319,7 @@ public:
 	void RemoveSectionByIdentity(FHierarchyElementIdentity SectionIdentity);
 	
 	virtual bool Modify(bool bAlwaysMarkDirty = true) override;
+	virtual void Serialize(FStructuredArchive::FRecord Record) override;
 
 	virtual void EmptyAllData();
 protected:
